Unreal Editor Remote Session now advertises session using NetworkServiceDiscovery plugin providing cross-platform mDNS/DNS-SD service discovery using native platform APIs

#Synced-CL 52439447

diff --git a/Engine/Plugins/Experimental/NetworkServiceDiscovery/NetworkServiceDiscovery.uplugin b/Engine/Plugins/Experimental/NetworkServiceDiscovery/NetworkServiceDiscovery.uplugin
new file mode 100644
--- /dev/null
+++ b/Engine/Plugins/Experimental/NetworkServiceDiscovery/NetworkServiceDiscovery.uplugin
@@ -0,0 +1,30 @@
+{
+	"FileVersion": 3,
+	"Version": 1,
+	"VersionName": "1.0",
+	"FriendlyName": "Network Service Discovery",
+	"Description": "Cross-platform mDNS/DNS-SD service discovery using native platform APIs (NSNetServiceBrowser on Apple, NsdManager on Android, native DNS-SD via dnsapi on Windows)",
+	"Category": "Experimental",
+	"CreatedBy": "Epic Games, Inc.",
+	"CreatedByURL": "https://epicgames.com",
+	"DocsURL": "",
+	"MarketplaceURL": "",
+	"SupportURL": "",
+	"CanContainContent": false,
+	"IsBetaVersion": true,
+	"Installed": false,
+	"EnabledByDefault": false,
+	"Modules": [
+		{
+			"Name": "NetworkServiceDiscovery",
+			"Type": "Runtime",
+			"LoadingPhase": "Default",
+			"PlatformAllowList": [
+				"Win64",
+				"Mac",
+				"IOS",
+				"Android"
+			]
+		}
+	]
+}
diff --git a/Engine/Plugins/Experimental/NetworkServiceDiscovery/Source/NetworkServiceDiscovery/NetworkServiceDiscovery.Build.cs b/Engine/Plugins/Experimental/NetworkServiceDiscovery/Source/NetworkServiceDiscovery/NetworkServiceDiscovery.Build.cs
new file mode 100644
--- /dev/null
+++ b/Engine/Plugins/Experimental/NetworkServiceDiscovery/Source/NetworkServiceDiscovery/NetworkServiceDiscovery.Build.cs
@@ -0,0 +1,47 @@
+// Copyright Epic Games, Inc. All Rights Reserved.
+
+using System.IO;
+using UnrealBuildTool;
+
+public class NetworkServiceDiscovery : ModuleRules
+{
+	public NetworkServiceDiscovery(ReadOnlyTargetRules Target) : base(Target)
+	{
+		PCHUsage = ModuleRules.PCHUsageMode.UseExplicitOrSharedPCHs;
+
+		PublicDependencyModuleNames.AddRange(
+			new string[]
+			{
+				"Core",
+			}
+		);
+
+		PrivateDependencyModuleNames.AddRange(
+			new string[]
+			{
+				"CoreUObject",
+				"Engine",
+			}
+		);
+
+		if (Target.Platform == UnrealTargetPlatform.Android)
+		{
+			PrivateDependencyModuleNames.Add("Launch");
+
+			string PluginPath = Utils.MakePathRelativeTo(ModuleDirectory, Target.RelativeEnginePath);
+			AdditionalPropertiesForReceipt.Add("AndroidPlugin",
+				Path.Combine(PluginPath, "NetworkServiceDiscovery_APL_Android.xml"));
+		}
+
+		if (Target.Platform == UnrealTargetPlatform.Win64)
+		{
+			PublicDefinitions.Add("WITH_WINDOWS_DNSSD=1");
+			// dnsapi.dll is loaded dynamically at runtime — no link-time dependency.
+			// This allows graceful fallback on Windows versions that lack DNS-SD APIs.
+		}
+		else
+		{
+			PublicDefinitions.Add("WITH_WINDOWS_DNSSD=0");
+		}
+	}
+}
diff --git a/Engine/Plugins/Experimental/NetworkServiceDiscovery/Source/NetworkServiceDiscovery/NetworkServiceDiscovery_APL_Android.xml b/Engine/Plugins/Experimental/NetworkServiceDiscovery/Source/NetworkServiceDiscovery/NetworkServiceDiscovery_APL_Android.xml
new file mode 100644
--- /dev/null
+++ b/Engine/Plugins/Experimental/NetworkServiceDiscovery/Source/NetworkServiceDiscovery/NetworkServiceDiscovery_APL_Android.xml
@@ -0,0 +1,34 @@
+<?xml version="1.0" encoding="utf-8"?>
+<!--NetworkServiceDiscovery plugin additions for Android-->
+<root xmlns:android="http://schemas.android.com/apk/res/android">
+	<init>
+		<log text="NetworkServiceDiscovery Android init"/>
+	</init>
+
+	<prebuildCopies>
+		<log text="Copying NetworkServiceDiscoveryHelper java file"/>
+		<copyFile src="$S(PluginDir)/Private/Android/Java/NetworkServiceDiscoveryHelper.java"
+				  dst="$S(BuildDir)/src/com/epicgames/unreal/NetworkServiceDiscoveryHelper.java" />
+	</prebuildCopies>
+
+	<proguardAdditions>
+		<insert>
+-keep class com.epicgames.unreal.NetworkServiceDiscoveryHelper {
+	public *;
+	private static native *;
+}
+		</insert>
+	</proguardAdditions>
+
+	<gameActivityImportAdditions>
+		<insert>
+import com.epicgames.unreal.NetworkServiceDiscoveryHelper;
+		</insert>
+	</gameActivityImportAdditions>
+
+	<gameActivityOnCreateAdditions>
+		<insert>
+		NetworkServiceDiscoveryHelper.Initialize(getApplicationContext());
+		</insert>
+	</gameActivityOnCreateAdditions>
+</root>
diff --git a/Engine/Plugins/Experimental/NetworkServiceDiscovery/Source/NetworkServiceDiscovery/Private/Android/Java/NetworkServiceDiscoveryHelper.java b/Engine/Plugins/Experimental/NetworkServiceDiscovery/Source/NetworkServiceDiscovery/Private/Android/Java/NetworkServiceDiscoveryHelper.java
new file mode 100644
--- /dev/null
+++ b/Engine/Plugins/Experimental/NetworkServiceDiscovery/Source/NetworkServiceDiscovery/Private/Android/Java/NetworkServiceDiscoveryHelper.java
@@ -0,0 +1,342 @@
+// Copyright Epic Games, Inc. All Rights Reserved.
+
+package com.epicgames.unreal;
+
+import android.content.Context;
+import android.net.nsd.NsdManager;
+import android.net.nsd.NsdServiceInfo;
+import android.util.Log;
+
+import java.net.InetAddress;
+import java.util.HashMap;
+import java.util.Map;
+import java.util.concurrent.ConcurrentHashMap;
+
+/**
+ * Java helper for Android NSD (Network Service Discovery) via NsdManager.
+ * Called from C++ via JNI, and calls back to C++ via native methods.
+ * Supports multiple simultaneous service registrations.
+ */
+public class NetworkServiceDiscoveryHelper
+{
+	private static final String TAG = "NSDHelper";
+
+	private static NetworkServiceDiscoveryHelper Instance;
+
+	private Context AppContext;
+	private NsdManager NsdMgr;
+
+	/** Active registrations keyed by service name */
+	private ConcurrentHashMap<String, NsdManager.RegistrationListener> RegistrationListeners = new ConcurrentHashMap<>();
+
+	private NsdManager.DiscoveryListener DiscoveryListener;
+	private boolean bIsDiscovering = false;
+
+	// ------------------------------------------------------------------
+	//  Native callbacks (implemented in C++)
+	// ------------------------------------------------------------------
+
+	private static native void nativeOnServiceFound(String serviceName, String serviceType);
+	private static native void nativeOnServiceLost(String serviceName, String serviceType);
+	private static native void nativeOnServiceResolved(String serviceName, String serviceType,
+		String hostAddress, int port, String[] txtKeys, String[] txtValues);
+	private static native void nativeOnServiceRegistered(String serviceName, String serviceType, int port);
+	private static native void nativeOnError(String errorMessage);
+
+	// ------------------------------------------------------------------
+	//  Initialization
+	// ------------------------------------------------------------------
+
+	public static void Initialize(Context context)
+	{
+		if (Instance == null)
+		{
+			Instance = new NetworkServiceDiscoveryHelper();
+			Instance.AppContext = context.getApplicationContext();
+			Instance.NsdMgr = (NsdManager) Instance.AppContext.getSystemService(Context.NSD_SERVICE);
+			Log.d(TAG, "NetworkServiceDiscoveryHelper initialized");
+		}
+	}
+
+	public static NetworkServiceDiscoveryHelper GetInstance()
+	{
+		return Instance;
+	}
+
+	// ------------------------------------------------------------------
+	//  Service Registration (supports multiple simultaneous registrations)
+	// ------------------------------------------------------------------
+
+	public boolean registerService(String serviceName, String serviceType, int port,
+		String[] txtKeys, String[] txtValues)
+	{
+		if (NsdMgr == null)
+		{
+			nativeOnError("NsdManager not available");
+			return false;
+		}
+
+		// Unregister any existing service with the same name
+		unregisterService(serviceName);
+
+		NsdServiceInfo serviceInfo = new NsdServiceInfo();
+		serviceInfo.setServiceName(serviceName);
+		serviceInfo.setServiceType(serviceType);
+		serviceInfo.setPort(port);
+
+		if (txtKeys != null && txtValues != null)
+		{
+			for (int i = 0; i < txtKeys.length && i < txtValues.length; i++)
+			{
+				serviceInfo.setAttribute(txtKeys[i], txtValues[i]);
+			}
+		}
+
+		final String regName = serviceName;
+		final int regPort = port;
+
+		NsdManager.RegistrationListener listener = new NsdManager.RegistrationListener()
+		{
+			@Override
+			public void onServiceRegistered(NsdServiceInfo info)
+			{
+				Log.d(TAG, "Service registered: " + info.getServiceName());
+				nativeOnServiceRegistered(info.getServiceName(), info.getServiceType(), regPort);
+			}
+
+			@Override
+			public void onRegistrationFailed(NsdServiceInfo info, int errorCode)
+			{
+				Log.e(TAG, "Registration failed for '" + regName + "': " + errorCode);
+				RegistrationListeners.remove(regName);
+				nativeOnError("NSD registration failed for '" + regName + "' with error code: " + errorCode);
+			}
+
+			@Override
+			public void onServiceUnregistered(NsdServiceInfo info)
+			{
+				Log.d(TAG, "Service unregistered: " + info.getServiceName());
+				RegistrationListeners.remove(regName);
+			}
+
+			@Override
+			public void onUnregistrationFailed(NsdServiceInfo info, int errorCode)
+			{
+				Log.e(TAG, "Unregistration failed for '" + regName + "': " + errorCode);
+				nativeOnError("NSD unregistration failed for '" + regName + "' with error code: " + errorCode);
+			}
+		};
+
+		try
+		{
+			NsdMgr.registerService(serviceInfo, NsdManager.PROTOCOL_DNS_SD, listener);
+			RegistrationListeners.put(serviceName, listener);
+			return true;
+		}
+		catch (Exception e)
+		{
+			Log.e(TAG, "Exception registering service: " + e.getMessage());
+			nativeOnError("Exception registering service: " + e.getMessage());
+			return false;
+		}
+	}
+
+	public void unregisterService(String serviceName)
+	{
+		if (NsdMgr == null) return;
+
+		if (serviceName == null || serviceName.isEmpty())
+		{
+			// Unregister all
+			for (Map.Entry<String, NsdManager.RegistrationListener> entry : RegistrationListeners.entrySet())
+			{
+				try
+				{
+					NsdMgr.unregisterService(entry.getValue());
+				}
+				catch (Exception e)
+				{
+					Log.w(TAG, "Exception unregistering service '" + entry.getKey() + "': " + e.getMessage());
+				}
+			}
+			RegistrationListeners.clear();
+		}
+		else
+		{
+			NsdManager.RegistrationListener listener = RegistrationListeners.remove(serviceName);
+			if (listener != null)
+			{
+				try
+				{
+					NsdMgr.unregisterService(listener);
+				}
+				catch (Exception e)
+				{
+					Log.w(TAG, "Exception unregistering service '" + serviceName + "': " + e.getMessage());
+				}
+			}
+		}
+	}
+
+	public boolean isServiceRegistered(String serviceName)
+	{
+		if (serviceName == null || serviceName.isEmpty())
+		{
+			return !RegistrationListeners.isEmpty();
+		}
+		return RegistrationListeners.containsKey(serviceName);
+	}
+
+	// ------------------------------------------------------------------
+	//  Service Discovery
+	// ------------------------------------------------------------------
+
+	public boolean startDiscovery(String serviceType)
+	{
+		if (NsdMgr == null)
+		{
+			nativeOnError("NsdManager not available");
+			return false;
+		}
+
+		stopDiscovery();
+
+		DiscoveryListener = new NsdManager.DiscoveryListener()
+		{
+			@Override
+			public void onDiscoveryStarted(String regType)
+			{
+				Log.d(TAG, "Discovery started for: " + regType);
+				bIsDiscovering = true;
+			}
+
+			@Override
+			public void onServiceFound(NsdServiceInfo info)
+			{
+				Log.d(TAG, "Service found: " + info.getServiceName() + " type: " + info.getServiceType());
+				nativeOnServiceFound(info.getServiceName(), info.getServiceType());
+			}
+
+			@Override
+			public void onServiceLost(NsdServiceInfo info)
+			{
+				Log.d(TAG, "Service lost: " + info.getServiceName());
+				nativeOnServiceLost(info.getServiceName(), info.getServiceType());
+			}
+
+			@Override
+			public void onDiscoveryStopped(String serviceType)
+			{
+				Log.d(TAG, "Discovery stopped for: " + serviceType);
+				bIsDiscovering = false;
+			}
+
+			@Override
+			public void onStartDiscoveryFailed(String serviceType, int errorCode)
+			{
+				Log.e(TAG, "Start discovery failed: " + errorCode);
+				bIsDiscovering = false;
+				nativeOnError("NSD start discovery failed with error code: " + errorCode);
+			}
+
+			@Override
+			public void onStopDiscoveryFailed(String serviceType, int errorCode)
+			{
+				Log.e(TAG, "Stop discovery failed: " + errorCode);
+				nativeOnError("NSD stop discovery failed with error code: " + errorCode);
+			}
+		};
+
+		try
+		{
+			NsdMgr.discoverServices(serviceType, NsdManager.PROTOCOL_DNS_SD, DiscoveryListener);
+			return true;
+		}
+		catch (Exception e)
+		{
+			Log.e(TAG, "Exception starting discovery: " + e.getMessage());
+			nativeOnError("Exception starting discovery: " + e.getMessage());
+			DiscoveryListener = null;
+			return false;
+		}
+	}
+
+	public void stopDiscovery()
+	{
+		if (DiscoveryListener != null && NsdMgr != null)
+		{
+			try
+			{
+				NsdMgr.stopServiceDiscovery(DiscoveryListener);
+			}
+			catch (Exception e)
+			{
+				Log.w(TAG, "Exception stopping discovery: " + e.getMessage());
+			}
+			DiscoveryListener = null;
+		}
+		bIsDiscovering = false;
+	}
+
+	public boolean isDiscovering()
+	{
+		return bIsDiscovering;
+	}
+
+	// ------------------------------------------------------------------
+	//  Service Resolution
+	// ------------------------------------------------------------------
+
+	public void resolveService(String serviceName, String serviceType)
+	{
+		if (NsdMgr == null)
+		{
+			nativeOnError("NsdManager not available for resolution");
+			return;
+		}
+
+		NsdServiceInfo serviceInfo = new NsdServiceInfo();
+		serviceInfo.setServiceName(serviceName);
+		serviceInfo.setServiceType(serviceType);
+
+		NsdMgr.resolveService(serviceInfo, new NsdManager.ResolveListener()
+		{
+			@Override
+			public void onResolveFailed(NsdServiceInfo info, int errorCode)
+			{
+				Log.e(TAG, "Resolve failed for " + info.getServiceName() + ": " + errorCode);
+				nativeOnError("NSD resolve failed for '" + info.getServiceName() + "' with error code: " + errorCode);
+			}
+
+			@Override
+			public void onServiceResolved(NsdServiceInfo info)
+			{
+				Log.d(TAG, "Resolved: " + info.getServiceName() + " -> " + info.getHost() + ":" + info.getPort());
+
+				InetAddress host = info.getHost();
+				String hostAddress = (host != null) ? host.getHostAddress() : "";
+
+				Map<String, byte[]> attributes = info.getAttributes();
+				String[] keys = new String[attributes.size()];
+				String[] values = new String[attributes.size()];
+				int index = 0;
+				for (Map.Entry<String, byte[]> entry : attributes.entrySet())
+				{
+					keys[index] = entry.getKey();
+					byte[] valueBytes = entry.getValue();
+					values[index] = (valueBytes != null) ? new String(valueBytes, java.nio.charset.StandardCharsets.UTF_8) : "";
+					index++;
+				}
+
+				nativeOnServiceResolved(
+					info.getServiceName(),
+					info.getServiceType(),
+					hostAddress,
+					info.getPort(),
+					keys,
+					values
+				);
+			}
+		});
+	}
+}
diff --git a/Engine/Plugins/Experimental/NetworkServiceDiscovery/Source/NetworkServiceDiscovery/Private/Android/NetworkServiceDiscoveryAndroid.cpp b/Engine/Plugins/Experimental/NetworkServiceDiscovery/Source/NetworkServiceDiscovery/Private/Android/NetworkServiceDiscoveryAndroid.cpp
new file mode 100644
--- /dev/null
+++ b/Engine/Plugins/Experimental/NetworkServiceDiscovery/Source/NetworkServiceDiscovery/Private/Android/NetworkServiceDiscoveryAndroid.cpp
@@ -0,0 +1,539 @@
+// Copyright Epic Games, Inc. All Rights Reserved.
+
+#include "Android/NetworkServiceDiscoveryAndroid.h"
+#include "NetworkServiceDiscoveryModule.h"
+#include "Async/Async.h"
+
+#if PLATFORM_ANDROID
+
+#include "Android/AndroidApplication.h"
+#include "Android/AndroidJNI.h"
+
+FNetworkServiceDiscoveryAndroid* FNetworkServiceDiscoveryAndroid::Instance = nullptr;
+
+FNetworkServiceDiscoveryAndroid::FNetworkServiceDiscoveryAndroid()
+{
+	Instance = this;
+	InitJNI();
+}
+
+FNetworkServiceDiscoveryAndroid::~FNetworkServiceDiscoveryAndroid()
+{
+	if (bJNIInitialized)
+	{
+		UnregisterService(FString());
+		StopDiscovery();
+	}
+
+	// Clean up any JNI global refs that were created (even if InitJNI failed partway)
+	JNIEnv* Env = FAndroidApplication::GetJavaEnv();
+	if (Env)
+	{
+		if (HelperInstance)
+		{
+			Env->DeleteGlobalRef(HelperInstance);
+			HelperInstance = nullptr;
+		}
+		if (HelperClass)
+		{
+			Env->DeleteGlobalRef(HelperClass);
+			HelperClass = nullptr;
+		}
+	}
+
+	if (Instance == this)
+	{
+		Instance = nullptr;
+	}
+}
+
+bool FNetworkServiceDiscoveryAndroid::InitJNI()
+{
+	JNIEnv* Env = FAndroidApplication::GetJavaEnv();
+	if (!Env)
+	{
+		UE_LOG(LogNetworkServiceDiscovery, Error, TEXT("Android: Failed to get JNI environment"));
+		return false;
+	}
+
+	// Find our helper class
+	jclass LocalClass = FAndroidApplication::FindJavaClass("com/epicgames/unreal/NetworkServiceDiscoveryHelper");
+	if (!LocalClass)
+	{
+		UE_LOG(LogNetworkServiceDiscovery, Error, TEXT("Android: Failed to find NetworkServiceDiscoveryHelper class"));
+		Env->ExceptionClear();
+		return false;
+	}
+	HelperClass = (jclass)Env->NewGlobalRef(LocalClass);
+	Env->DeleteLocalRef(LocalClass);
+
+	// Initialize the helper with the application context
+	jmethodID InitMethod = Env->GetStaticMethodID(HelperClass, "Initialize", "(Landroid/content/Context;)V");
+	Env->ExceptionClear();
+	if (InitMethod)
+	{
+		jobject Activity = FAndroidApplication::GetGameActivityThis();
+		Env->CallStaticVoidMethod(HelperClass, InitMethod, Activity);
+		Env->ExceptionClear();
+	}
+
+	// Get the singleton instance
+	jmethodID GetInstanceMethod = Env->GetStaticMethodID(HelperClass, "GetInstance", "()Lcom/epicgames/unreal/NetworkServiceDiscoveryHelper;");
+	Env->ExceptionClear();
+	if (GetInstanceMethod)
+	{
+		jobject LocalInstance = Env->CallStaticObjectMethod(HelperClass, GetInstanceMethod);
+		Env->ExceptionClear();
+		if (LocalInstance)
+		{
+			HelperInstance = Env->NewGlobalRef(LocalInstance);
+			Env->DeleteLocalRef(LocalInstance);
+		}
+	}
+
+	if (!HelperInstance)
+	{
+		UE_LOG(LogNetworkServiceDiscovery, Error, TEXT("Android: Failed to get helper instance"));
+		// Clean up HelperClass global ref since we're failing
+		Env->DeleteGlobalRef(HelperClass);
+		HelperClass = nullptr;
+		return false;
+	}
+
+	// Cache method IDs — validate each one; a null method ID would crash on use
+	RegisterServiceMethod = Env->GetMethodID(HelperClass, "registerService",
+		"(Ljava/lang/String;Ljava/lang/String;I[Ljava/lang/String;[Ljava/lang/String;)Z");
+	Env->ExceptionClear();
+	UnregisterServiceMethod = Env->GetMethodID(HelperClass, "unregisterService", "(Ljava/lang/String;)V");
+	Env->ExceptionClear();
+	StartDiscoveryMethod = Env->GetMethodID(HelperClass, "startDiscovery", "(Ljava/lang/String;)Z");
+	Env->ExceptionClear();
+	StopDiscoveryMethod = Env->GetMethodID(HelperClass, "stopDiscovery", "()V");
+	Env->ExceptionClear();
+	ResolveServiceMethod = Env->GetMethodID(HelperClass, "resolveService",
+		"(Ljava/lang/String;Ljava/lang/String;)V");
+	Env->ExceptionClear();
+	IsServiceRegisteredMethod = Env->GetMethodID(HelperClass, "isServiceRegistered", "(Ljava/lang/String;)Z");
+	Env->ExceptionClear();
+	IsDiscoveringMethod = Env->GetMethodID(HelperClass, "isDiscovering", "()Z");
+	Env->ExceptionClear();
+
+	if (!RegisterServiceMethod || !UnregisterServiceMethod || !StartDiscoveryMethod ||
+		!StopDiscoveryMethod || !ResolveServiceMethod || !IsServiceRegisteredMethod || !IsDiscoveringMethod)
+	{
+		UE_LOG(LogNetworkServiceDiscovery, Error, TEXT("Android: Failed to resolve one or more JNI method IDs. Check ProGuard rules and Java method signatures."));
+		// Clean up global refs since we're failing
+		Env->DeleteGlobalRef(HelperInstance);
+		HelperInstance = nullptr;
+		Env->DeleteGlobalRef(HelperClass);
+		HelperClass = nullptr;
+		return false;
+	}
+
+	bJNIInitialized = true;
+	UE_LOG(LogNetworkServiceDiscovery, Log, TEXT("Android: JNI initialized successfully"));
+	return true;
+}
+
+// --- Registration ---
+
+bool FNetworkServiceDiscoveryAndroid::RegisterService(const FString& ServiceName, const FString& ServiceType, int32 Port, const TMap<FString, FString>& TxtRecord)
+{
+	if (!bJNIInitialized)
+	{
+		return false;
+	}
+
+	JNIEnv* Env = FAndroidApplication::GetJavaEnv();
+	if (!Env)
+	{
+		return false;
+	}
+
+	jstring JName = Env->NewStringUTF(TCHAR_TO_UTF8(*ServiceName));
+	jstring JType = Env->NewStringUTF(TCHAR_TO_UTF8(*ServiceType));
+
+	// Build TXT record arrays
+	jobjectArray JKeys = nullptr;
+	jobjectArray JValues = nullptr;
+
+	if (TxtRecord.Num() > 0)
+	{
+		jclass StringClass = Env->FindClass("java/lang/String");
+		if (StringClass)
+		{
+			JKeys = Env->NewObjectArray(TxtRecord.Num(), StringClass, nullptr);
+			JValues = Env->NewObjectArray(TxtRecord.Num(), StringClass, nullptr);
+
+			if (JKeys && JValues)
+			{
+				int32 Index = 0;
+				for (const auto& Pair : TxtRecord)
+				{
+					jstring JKey = Env->NewStringUTF(TCHAR_TO_UTF8(*Pair.Key));
+					jstring JValue = Env->NewStringUTF(TCHAR_TO_UTF8(*Pair.Value));
+					Env->SetObjectArrayElement(JKeys, Index, JKey);
+					Env->SetObjectArrayElement(JValues, Index, JValue);
+					Env->DeleteLocalRef(JKey);
+					Env->DeleteLocalRef(JValue);
+					Index++;
+				}
+			}
+			Env->DeleteLocalRef(StringClass);
+		}
+	}
+
+	jboolean Result = Env->CallBooleanMethod(HelperInstance, RegisterServiceMethod, JName, JType, (jint)Port, JKeys, JValues);
+
+	if (Env->ExceptionCheck())
+	{
+		Env->ExceptionClear();
+		UE_LOG(LogNetworkServiceDiscovery, Error, TEXT("Android: Exception in RegisterService"));
+		Result = JNI_FALSE;
+	}
+
+	Env->DeleteLocalRef(JName);
+	Env->DeleteLocalRef(JType);
+	if (JKeys) Env->DeleteLocalRef(JKeys);
+	if (JValues) Env->DeleteLocalRef(JValues);
+
+	return (bool)Result;
+}
+
+void FNetworkServiceDiscoveryAndroid::UnregisterService(const FString& ServiceName)
+{
+	if (!bJNIInitialized)
+	{
+		return;
+	}
+
+	JNIEnv* Env = FAndroidApplication::GetJavaEnv();
+	if (Env)
+	{
+		jstring JName = Env->NewStringUTF(TCHAR_TO_UTF8(*ServiceName));
+		Env->CallVoidMethod(HelperInstance, UnregisterServiceMethod, JName);
+		if (Env->ExceptionCheck())
+		{
+			Env->ExceptionClear();
+			UE_LOG(LogNetworkServiceDiscovery, Warning, TEXT("Android: Exception in UnregisterService"));
+		}
+		Env->DeleteLocalRef(JName);
+	}
+}
+
+bool FNetworkServiceDiscoveryAndroid::IsServiceRegistered(const FString& ServiceName) const
+{
+	if (!bJNIInitialized)
+	{
+		return false;
+	}
+
+	JNIEnv* Env = FAndroidApplication::GetJavaEnv();
+	if (Env)
+	{
+		jstring JName = Env->NewStringUTF(TCHAR_TO_UTF8(*ServiceName));
+		jboolean Result = Env->CallBooleanMethod(HelperInstance, IsServiceRegisteredMethod, JName);
+		Env->DeleteLocalRef(JName);
+		return (bool)Result;
+	}
+	return false;
+}
+
+// --- Discovery ---
+
+bool FNetworkServiceDiscoveryAndroid::StartDiscovery(const FString& ServiceType)
+{
+	if (!bJNIInitialized)
+	{
+		return false;
+	}
+
+	JNIEnv* Env = FAndroidApplication::GetJavaEnv();
+	if (!Env)
+	{
+		return false;
+	}
+
+	jstring JType = Env->NewStringUTF(TCHAR_TO_UTF8(*ServiceType));
+	jboolean Result = Env->CallBooleanMethod(HelperInstance, StartDiscoveryMethod, JType);
+
+	if (Env->ExceptionCheck())
+	{
+		Env->ExceptionClear();
+		UE_LOG(LogNetworkServiceDiscovery, Error, TEXT("Android: Exception in StartDiscovery"));
+		Result = JNI_FALSE;
+	}
+	Env->DeleteLocalRef(JType);
+
+	return (bool)Result;
+}
+
+void FNetworkServiceDiscoveryAndroid::StopDiscovery()
+{
+	if (!bJNIInitialized)
+	{
+		return;
+	}
+
+	JNIEnv* Env = FAndroidApplication::GetJavaEnv();
+	if (Env)
+	{
+		Env->CallVoidMethod(HelperInstance, StopDiscoveryMethod);
+		if (Env->ExceptionCheck())
+		{
+			Env->ExceptionClear();
+			UE_LOG(LogNetworkServiceDiscovery, Warning, TEXT("Android: Exception in StopDiscovery"));
+		}
+	}
+
+	FScopeLock Lock(&ServicesLock);
+	DiscoveredServices.Empty();
+}
+
+bool FNetworkServiceDiscoveryAndroid::IsDiscovering() const
+{
+	if (!bJNIInitialized)
+	{
+		return false;
+	}
+
+	JNIEnv* Env = FAndroidApplication::GetJavaEnv();
+	if (Env)
+	{
+		return (bool)Env->CallBooleanMethod(HelperInstance, IsDiscoveringMethod);
+	}
+	return false;
+}
+
+void FNetworkServiceDiscoveryAndroid::ResolveService(const FNetworkServiceInfo& Service)
+{
+	if (!bJNIInitialized)
+	{
+		return;
+	}
+
+	JNIEnv* Env = FAndroidApplication::GetJavaEnv();
+	if (!Env)
+	{
+		return;
+	}
+
+	jstring JName = Env->NewStringUTF(TCHAR_TO_UTF8(*Service.ServiceName));
+	jstring JType = Env->NewStringUTF(TCHAR_TO_UTF8(*Service.ServiceType));
+	Env->CallVoidMethod(HelperInstance, ResolveServiceMethod, JName, JType);
+	if (Env->ExceptionCheck())
+	{
+		Env->ExceptionClear();
+		UE_LOG(LogNetworkServiceDiscovery, Warning, TEXT("Android: Exception in ResolveService"));
+	}
+	Env->DeleteLocalRef(JName);
+	Env->DeleteLocalRef(JType);
+}
+
+TArray<FNetworkServiceInfo> FNetworkServiceDiscoveryAndroid::GetDiscoveredServices() const
+{
+	FScopeLock Lock(&ServicesLock);
+	return DiscoveredServices;
+}
+
+// --- Callbacks from JNI (dispatched to game thread) ---
+
+void FNetworkServiceDiscoveryAndroid::HandleServiceFound(const FString& ServiceName, const FString& ServiceType)
+{
+	FNetworkServiceInfo Info;
+	Info.ServiceName = ServiceName;
+	Info.ServiceType = ServiceType;
+
+	{
+		FScopeLock Lock(&ServicesLock);
+		DiscoveredServices.Add(Info);
+	}
+	OnServiceFoundDelegate.Broadcast(Info);
+}
+
+void FNetworkServiceDiscoveryAndroid::HandleServiceLost(const FString& ServiceName, const FString& ServiceType)
+{
+	FNetworkServiceInfo Info;
+	Info.ServiceName = ServiceName;
+	Info.ServiceType = ServiceType;
+
+	{
+		FScopeLock Lock(&ServicesLock);
+		DiscoveredServices.RemoveAll([&ServiceName, &ServiceType](const FNetworkServiceInfo& Existing)
+		{
+			return Existing.ServiceName == ServiceName && Existing.ServiceType == ServiceType;
+		});
+	}
+	OnServiceLostDelegate.Broadcast(Info);
+}
+
+void FNetworkServiceDiscoveryAndroid::HandleServiceResolved(const FNetworkServiceInfo& Service)
+{
+	{
+		FScopeLock Lock(&ServicesLock);
+		for (FNetworkServiceInfo& Existing : DiscoveredServices)
+		{
+			if (Existing.ServiceName == Service.ServiceName && Existing.ServiceType == Service.ServiceType)
+			{
+				Existing = Service;
+				break;
+			}
+		}
+	}
+	OnServiceResolvedDelegate.Broadcast(Service);
+}
+
+void FNetworkServiceDiscoveryAndroid::HandleServiceRegistered(const FString& ServiceName, const FString& ServiceType, int32 Port)
+{
+	FNetworkServiceInfo Info;
+	Info.ServiceName = ServiceName;
+	Info.ServiceType = ServiceType;
+	Info.Port = Port;
+	Info.bIsResolved = true;
+	OnServiceRegisteredDelegate.Broadcast(Info);
+}
+
+void FNetworkServiceDiscoveryAndroid::HandleError(const FString& ErrorMessage)
+{
+	OnDiscoveryErrorDelegate.Broadcast(ErrorMessage);
+}
+
+// ============================================================================
+//  JNI native method implementations
+// ============================================================================
+
+extern "C"
+{
+
+JNIEXPORT void JNICALL Java_com_epicgames_unreal_NetworkServiceDiscoveryHelper_nativeOnServiceFound(
+	JNIEnv* Env, jclass Clazz, jstring JServiceName, jstring JServiceType)
+{
+	const char* NameChars = Env->GetStringUTFChars(JServiceName, nullptr);
+	const char* TypeChars = Env->GetStringUTFChars(JServiceType, nullptr);
+	FString ServiceName = UTF8_TO_TCHAR(NameChars);
+	FString ServiceType = UTF8_TO_TCHAR(TypeChars);
+	Env->ReleaseStringUTFChars(JServiceName, NameChars);
+	Env->ReleaseStringUTFChars(JServiceType, TypeChars);
+
+	AsyncTask(ENamedThreads::GameThread, [ServiceName, ServiceType]()
+	{
+		if (FNetworkServiceDiscoveryAndroid* Impl = FNetworkServiceDiscoveryAndroid::GetInstance())
+		{
+			Impl->HandleServiceFound(ServiceName, ServiceType);
+		}
+	});
+}
+
+JNIEXPORT void JNICALL Java_com_epicgames_unreal_NetworkServiceDiscoveryHelper_nativeOnServiceLost(
+	JNIEnv* Env, jclass Clazz, jstring JServiceName, jstring JServiceType)
+{
+	const char* NameChars = Env->GetStringUTFChars(JServiceName, nullptr);
+	const char* TypeChars = Env->GetStringUTFChars(JServiceType, nullptr);
+	FString ServiceName = UTF8_TO_TCHAR(NameChars);
+	FString ServiceType = UTF8_TO_TCHAR(TypeChars);
+	Env->ReleaseStringUTFChars(JServiceName, NameChars);
+	Env->ReleaseStringUTFChars(JServiceType, TypeChars);
+
+	AsyncTask(ENamedThreads::GameThread, [ServiceName, ServiceType]()
+	{
+		if (FNetworkServiceDiscoveryAndroid* Impl = FNetworkServiceDiscoveryAndroid::GetInstance())
+		{
+			Impl->HandleServiceLost(ServiceName, ServiceType);
+		}
+	});
+}
+
+JNIEXPORT void JNICALL Java_com_epicgames_unreal_NetworkServiceDiscoveryHelper_nativeOnServiceResolved(
+	JNIEnv* Env, jclass Clazz, jstring JServiceName, jstring JServiceType,
+	jstring JHostAddress, jint JPort, jobjectArray JTxtKeys, jobjectArray JTxtValues)
+{
+	const char* NameChars = Env->GetStringUTFChars(JServiceName, nullptr);
+	const char* TypeChars = Env->GetStringUTFChars(JServiceType, nullptr);
+	const char* AddrChars = Env->GetStringUTFChars(JHostAddress, nullptr);
+
+	FNetworkServiceInfo Info;
+	Info.ServiceName = UTF8_TO_TCHAR(NameChars);
+	Info.ServiceType = UTF8_TO_TCHAR(TypeChars);
+	Info.Address = UTF8_TO_TCHAR(AddrChars);
+	Info.Port = (int32)JPort;
+	Info.bIsResolved = true;
+
+	Env->ReleaseStringUTFChars(JServiceName, NameChars);
+	Env->ReleaseStringUTFChars(JServiceType, TypeChars);
+	Env->ReleaseStringUTFChars(JHostAddress, AddrChars);
+
+	// Extract TXT record
+	if (JTxtKeys && JTxtValues)
+	{
+		int32 Count = Env->GetArrayLength(JTxtKeys);
+		for (int32 i = 0; i < Count; i++)
+		{
+			jstring JKey = (jstring)Env->GetObjectArrayElement(JTxtKeys, i);
+			jstring JValue = (jstring)Env->GetObjectArrayElement(JTxtValues, i);
+
+			if (JKey && JValue)
+			{
+				const char* KeyChars = Env->GetStringUTFChars(JKey, nullptr);
+				const char* ValueChars = Env->GetStringUTFChars(JValue, nullptr);
+
+				if (KeyChars && ValueChars)
+				{
+					Info.TxtRecord.Add(UTF8_TO_TCHAR(KeyChars), UTF8_TO_TCHAR(ValueChars));
+				}
+
+				if (KeyChars) Env->ReleaseStringUTFChars(JKey, KeyChars);
+				if (ValueChars) Env->ReleaseStringUTFChars(JValue, ValueChars);
+			}
+			if (JKey) Env->DeleteLocalRef(JKey);
+			if (JValue) Env->DeleteLocalRef(JValue);
+		}
+	}
+
+	AsyncTask(ENamedThreads::GameThread, [Info]()
+	{
+		if (FNetworkServiceDiscoveryAndroid* Impl = FNetworkServiceDiscoveryAndroid::GetInstance())
+		{
+			Impl->HandleServiceResolved(Info);
+		}
+	});
+}
+
+JNIEXPORT void JNICALL Java_com_epicgames_unreal_NetworkServiceDiscoveryHelper_nativeOnServiceRegistered(
+	JNIEnv* Env, jclass Clazz, jstring JServiceName, jstring JServiceType, jint JPort)
+{
+	const char* NameChars = Env->GetStringUTFChars(JServiceName, nullptr);
+	const char* TypeChars = Env->GetStringUTFChars(JServiceType, nullptr);
+	FString ServiceName = UTF8_TO_TCHAR(NameChars);
+	FString ServiceType = UTF8_TO_TCHAR(TypeChars);
+	int32 Port = (int32)JPort;
+	Env->ReleaseStringUTFChars(JServiceName, NameChars);
+	Env->ReleaseStringUTFChars(JServiceType, TypeChars);
+
+	AsyncTask(ENamedThreads::GameThread, [ServiceName, ServiceType, Port]()
+	{
+		if (FNetworkServiceDiscoveryAndroid* Impl = FNetworkServiceDiscoveryAndroid::GetInstance())
+		{
+			Impl->HandleServiceRegistered(ServiceName, ServiceType, Port);
+		}
+	});
+}
+
+JNIEXPORT void JNICALL Java_com_epicgames_unreal_NetworkServiceDiscoveryHelper_nativeOnError(
+	JNIEnv* Env, jclass Clazz, jstring JErrorMessage)
+{
+	const char* MsgChars = Env->GetStringUTFChars(JErrorMessage, nullptr);
+	FString ErrorMessage = UTF8_TO_TCHAR(MsgChars);
+	Env->ReleaseStringUTFChars(JErrorMessage, MsgChars);
+
+	AsyncTask(ENamedThreads::GameThread, [ErrorMessage]()
+	{
+		if (FNetworkServiceDiscoveryAndroid* Impl = FNetworkServiceDiscoveryAndroid::GetInstance())
+		{
+			Impl->HandleError(ErrorMessage);
+		}
+	});
+}
+
+} // extern "C"
+
+#endif // PLATFORM_ANDROID
diff --git a/Engine/Plugins/Experimental/NetworkServiceDiscovery/Source/NetworkServiceDiscovery/Private/Android/NetworkServiceDiscoveryAndroid.h b/Engine/Plugins/Experimental/NetworkServiceDiscovery/Source/NetworkServiceDiscovery/Private/Android/NetworkServiceDiscoveryAndroid.h
new file mode 100644
--- /dev/null
+++ b/Engine/Plugins/Experimental/NetworkServiceDiscovery/Source/NetworkServiceDiscovery/Private/Android/NetworkServiceDiscoveryAndroid.h
@@ -0,0 +1,63 @@
+// Copyright Epic Games, Inc. All Rights Reserved.
+
+#pragma once
+
+#include "INetworkServiceDiscoveryPlatform.h"
+
+/**
+ * Android implementation of network service discovery using NsdManager via JNI.
+ */
+class FNetworkServiceDiscoveryAndroid : public INetworkServiceDiscoveryPlatform
+{
+public:
+	FNetworkServiceDiscoveryAndroid();
+	virtual ~FNetworkServiceDiscoveryAndroid();
+
+	// Non-copyable — JNI global refs are owned exclusively
+	FNetworkServiceDiscoveryAndroid(const FNetworkServiceDiscoveryAndroid&) = delete;
+	FNetworkServiceDiscoveryAndroid& operator=(const FNetworkServiceDiscoveryAndroid&) = delete;
+
+	// Registration
+	virtual bool RegisterService(const FString& ServiceName, const FString& ServiceType, int32 Port, const TMap<FString, FString>& TxtRecord) override;
+	virtual void UnregisterService(const FString& ServiceName) override;
+	virtual bool IsServiceRegistered(const FString& ServiceName) const override;
+
+	// Discovery
+	virtual bool StartDiscovery(const FString& ServiceType) override;
+	virtual void StopDiscovery() override;
+	virtual bool IsDiscovering() const override;
+	virtual void ResolveService(const FNetworkServiceInfo& Service) override;
+	virtual TArray<FNetworkServiceInfo> GetDiscoveredServices() const override;
+
+	// Called from JNI callbacks (on game thread)
+	void HandleServiceFound(const FString& ServiceName, const FString& ServiceType);
+	void HandleServiceLost(const FString& ServiceName, const FString& ServiceType);
+	void HandleServiceResolved(const FNetworkServiceInfo& Service);
+	void HandleServiceRegistered(const FString& ServiceName, const FString& ServiceType, int32 Port);
+	void HandleError(const FString& ErrorMessage);
+
+	/** Singleton accessor for JNI callbacks */
+	static FNetworkServiceDiscoveryAndroid* GetInstance() { return Instance; }
+
+private:
+	bool InitJNI();
+
+	/** JNI references — all default-initialized to null for safety on partial InitJNI failure */
+	jclass HelperClass = nullptr;
+	jobject HelperInstance = nullptr;
+	jmethodID RegisterServiceMethod = nullptr;
+	jmethodID UnregisterServiceMethod = nullptr;
+	jmethodID StartDiscoveryMethod = nullptr;
+	jmethodID StopDiscoveryMethod = nullptr;
+	jmethodID ResolveServiceMethod = nullptr;
+	jmethodID IsServiceRegisteredMethod = nullptr;
+	jmethodID IsDiscoveringMethod = nullptr;
+
+	bool bJNIInitialized = false;
+
+	/** Currently discovered services */
+	mutable FCriticalSection ServicesLock;
+	TArray<FNetworkServiceInfo> DiscoveredServices;
+
+	static FNetworkServiceDiscoveryAndroid* Instance;
+};
diff --git a/Engine/Plugins/Experimental/NetworkServiceDiscovery/Source/NetworkServiceDiscovery/Private/Apple/NetworkServiceDiscoveryApple.h b/Engine/Plugins/Experimental/NetworkServiceDiscovery/Source/NetworkServiceDiscovery/Private/Apple/NetworkServiceDiscoveryApple.h
new file mode 100644
--- /dev/null
+++ b/Engine/Plugins/Experimental/NetworkServiceDiscovery/Source/NetworkServiceDiscovery/Private/Apple/NetworkServiceDiscoveryApple.h
@@ -0,0 +1,51 @@
+// Copyright Epic Games, Inc. All Rights Reserved.
+
+#pragma once
+
+#include "INetworkServiceDiscoveryPlatform.h"
+
+@class FNSDServiceBrowserDelegate;
+@class FNSDServiceDelegate;
+
+/**
+ * Apple implementation of network service discovery using NSNetServiceBrowser / NSNetService.
+ * Works on both iOS and macOS.
+ */
+class FNetworkServiceDiscoveryApple : public INetworkServiceDiscoveryPlatform
+{
+public:
+	FNetworkServiceDiscoveryApple();
+	virtual ~FNetworkServiceDiscoveryApple();
+
+	// Registration
+	virtual bool RegisterService(const FString& ServiceName, const FString& ServiceType, int32 Port, const TMap<FString, FString>& TxtRecord) override;
+	virtual void UnregisterService(const FString& ServiceName) override;
+	virtual bool IsServiceRegistered(const FString& ServiceName) const override;
+
+	// Discovery
+	virtual bool StartDiscovery(const FString& ServiceType) override;
+	virtual void StopDiscovery() override;
+	virtual bool IsDiscovering() const override;
+	virtual void ResolveService(const FNetworkServiceInfo& Service) override;
+	virtual TArray<FNetworkServiceInfo> GetDiscoveredServices() const override;
+
+	// Called from Objective-C delegates
+	void HandleServiceFound(const FNetworkServiceInfo& Service);
+	void HandleServiceLost(const FNetworkServiceInfo& Service);
+	void HandleServiceResolved(const FNetworkServiceInfo& Service);
+	void HandleServiceRegistered(const FNetworkServiceInfo& Service);
+	void HandleError(const FString& ErrorMessage);
+
+private:
+	/** Objective-C delegate for NSNetServiceBrowser callbacks */
+	FNSDServiceBrowserDelegate* BrowserDelegate;
+
+	/** Published services keyed by service name */
+	TMap<FString, FNSDServiceDelegate*> PublishDelegates;
+
+	/** Currently discovered services */
+	mutable FCriticalSection ServicesLock;
+	TArray<FNetworkServiceInfo> DiscoveredServices;
+
+	bool bIsDiscovering;
+};
diff --git a/Engine/Plugins/Experimental/NetworkServiceDiscovery/Source/NetworkServiceDiscovery/Private/Apple/NetworkServiceDiscoveryApple.mm b/Engine/Plugins/Experimental/NetworkServiceDiscovery/Source/NetworkServiceDiscovery/Private/Apple/NetworkServiceDiscoveryApple.mm
new file mode 100644
--- /dev/null
+++ b/Engine/Plugins/Experimental/NetworkServiceDiscovery/Source/NetworkServiceDiscovery/Private/Apple/NetworkServiceDiscoveryApple.mm
@@ -0,0 +1,535 @@
+// Copyright Epic Games, Inc. All Rights Reserved.
+
+#include "Apple/NetworkServiceDiscoveryApple.h"
+#include "NetworkServiceDiscoveryModule.h"
+#include "Async/Async.h"
+
+// CarbonCore (pulled in by Foundation.h) defines its own 'struct FVector' which
+// clashes with UE's 'using FVector = ...' alias. Redirect any use of 'FVector'
+// inside the Apple headers to a harmless name so the two declarations never meet.
+#pragma push_macro("FVector")
+#define FVector FVectorWorkaround
+
+#include "Apple/PreAppleSystemHeaders.h"
+#import <Foundation/Foundation.h>
+#import <arpa/inet.h>
+#include "Apple/PostAppleSystemHeaders.h"
+
+#undef FVector
+#pragma pop_macro("FVector")
+
+// ============================================================================
+//  Objective-C delegate for NSNetServiceBrowser (discovery)
+// ============================================================================
+
+@interface FNSDServiceBrowserDelegate : NSObject <NSNetServiceBrowserDelegate, NSNetServiceDelegate>
+{
+	FNetworkServiceDiscoveryApple* Owner;
+	NSNetServiceBrowser* Browser;
+	NSMutableArray<NSNetService*>* ResolvingServices;
+}
+
+- (instancetype)initWithOwner:(FNetworkServiceDiscoveryApple*)InOwner;
+- (void)startBrowsingForType:(NSString*)ServiceType;
+- (void)stopBrowsing;
+- (void)resolveService:(NSString*)ServiceName ofType:(NSString*)ServiceType;
+
+@end
+
+@implementation FNSDServiceBrowserDelegate
+
+- (instancetype)initWithOwner:(FNetworkServiceDiscoveryApple*)InOwner
+{
+	self = [super init];
+	if (self)
+	{
+		Owner = InOwner;
+		Browser = [[NSNetServiceBrowser alloc] init];
+		Browser.delegate = self;
+		ResolvingServices = [[NSMutableArray alloc] init];
+	}
+	return self;
+}
+
+- (void)dealloc
+{
+	[self stopBrowsing];
+	Browser.delegate = nil;
+	[Browser release];
+	[ResolvingServices release];
+	[super dealloc];
+}
+
+- (void)startBrowsingForType:(NSString*)ServiceType
+{
+	// Strip trailing dot from service type — iOS NSBonjourServices plist entries
+	// don't include it, and a mismatch causes silent discovery failure.
+	NSString* CleanType = ServiceType;
+	if ([CleanType hasSuffix:@"."])
+	{
+		CleanType = [CleanType substringToIndex:[CleanType length] - 1];
+	}
+
+	// Must create and start the browser on the main thread to ensure the
+	// run loop delivers delegate callbacks. On iOS, UE's game thread may
+	// not pump NSRunLoop.
+	dispatch_async(dispatch_get_main_queue(), ^{
+		[Browser scheduleInRunLoop:[NSRunLoop mainRunLoop] forMode:NSDefaultRunLoopMode];
+		[Browser searchForServicesOfType:CleanType inDomain:@"local."];
+	});
+}
+
+- (void)stopBrowsing
+{
+	[Browser stop];
+	for (NSNetService* Service in ResolvingServices)
+	{
+		[Service stop];
+		Service.delegate = nil;
+	}
+	[ResolvingServices removeAllObjects];
+}
+
+- (void)resolveService:(NSString*)ServiceName ofType:(NSString*)ServiceType
+{
+	// Must resolve on the main thread so the run loop delivers delegate callbacks.
+	dispatch_async(dispatch_get_main_queue(), ^{
+		NSNetService* Service = [[NSNetService alloc] initWithDomain:@"local." type:ServiceType name:ServiceName];
+		Service.delegate = self;
+		[ResolvingServices addObject:Service];
+		[Service scheduleInRunLoop:[NSRunLoop mainRunLoop] forMode:NSDefaultRunLoopMode];
+		[Service resolveWithTimeout:10.0];
+		[Service release];
+	});
+}
+
+// --- NSNetServiceBrowserDelegate ---
+
+- (void)netServiceBrowserWillSearch:(NSNetServiceBrowser*)Browser
+{
+	UE_LOG(LogNetworkServiceDiscovery, Log, TEXT("Apple: Browser started searching"));
+}
+
+- (void)netServiceBrowserDidStopSearch:(NSNetServiceBrowser*)Browser
+{
+	UE_LOG(LogNetworkServiceDiscovery, Log, TEXT("Apple: Browser stopped searching"));
+}
+
+- (void)netServiceBrowser:(NSNetServiceBrowser*)Browser didFindService:(NSNetService*)Service moreComing:(BOOL)MoreComing
+{
+	FNetworkServiceInfo Info;
+	Info.ServiceName = FString(Service.name);
+	Info.ServiceType = FString(Service.type);
+	Info.HostName = Service.hostName ? FString(Service.hostName) : FString();
+
+	UE_LOG(LogNetworkServiceDiscovery, Log, TEXT("Apple: Found service '%s' of type '%s'"), *Info.ServiceName, *Info.ServiceType);
+
+	AsyncTask(ENamedThreads::GameThread, [WeakOwner = Owner, Info]()
+	{
+		if (WeakOwner)
+		{
+			WeakOwner->HandleServiceFound(Info);
+		}
+	});
+}
+
+- (void)netServiceBrowser:(NSNetServiceBrowser*)Browser didRemoveService:(NSNetService*)Service moreComing:(BOOL)MoreComing
+{
+	FNetworkServiceInfo Info;
+	Info.ServiceName = FString(Service.name);
+	Info.ServiceType = FString(Service.type);
+
+	UE_LOG(LogNetworkServiceDiscovery, Log, TEXT("Apple: Lost service '%s'"), *Info.ServiceName);
+
+	AsyncTask(ENamedThreads::GameThread, [WeakOwner = Owner, Info]()
+	{
+		if (WeakOwner)
+		{
+			WeakOwner->HandleServiceLost(Info);
+		}
+	});
+}
+
+- (void)netServiceBrowser:(NSNetServiceBrowser*)Browser didNotSearch:(NSDictionary<NSString*, NSNumber*>*)ErrorDict
+{
+	NSNumber* ErrorCode = ErrorDict[NSNetServicesErrorCode];
+	FString ErrorMsg = FString::Printf(TEXT("Apple: Browse failed with error code %d"), ErrorCode ? [ErrorCode intValue] : -1);
+
+	UE_LOG(LogNetworkServiceDiscovery, Error, TEXT("%s"), *ErrorMsg);
+
+	AsyncTask(ENamedThreads::GameThread, [WeakOwner = Owner, ErrorMsg]()
+	{
+		if (WeakOwner)
+		{
+			WeakOwner->HandleError(ErrorMsg);
+		}
+	});
+}
+
+// --- NSNetServiceDelegate (for resolution) ---
+
+- (void)netServiceDidResolveAddress:(NSNetService*)Service
+{
+	FNetworkServiceInfo Info;
+	Info.ServiceName = FString(Service.name);
+	Info.ServiceType = FString(Service.type);
+	Info.HostName = Service.hostName ? FString(Service.hostName) : FString();
+	Info.Port = (int32)Service.port;
+	Info.bIsResolved = true;
+
+	// Extract IPv4 address from the resolved addresses
+	for (NSData* AddressData in Service.addresses)
+	{
+		const struct sockaddr* SockAddr = (const struct sockaddr*)[AddressData bytes];
+		if (SockAddr->sa_family == AF_INET)
+		{
+			const struct sockaddr_in* Addr4 = (const struct sockaddr_in*)SockAddr;
+			char AddrStr[INET_ADDRSTRLEN];
+			inet_ntop(AF_INET, &Addr4->sin_addr, AddrStr, sizeof(AddrStr));
+			Info.Address = FString(ANSI_TO_TCHAR(AddrStr));
+			break;
+		}
+	}
+
+	// Extract TXT record
+	NSData* TxtData = Service.TXTRecordData;
+	if (TxtData)
+	{
+		NSDictionary<NSString*, NSData*>* TxtDict = [NSNetService dictionaryFromTXTRecordData:TxtData];
+		for (NSString* Key in TxtDict)
+		{
+			NSData* ValueData = TxtDict[Key];
+			NSString* ValueStr = [[NSString alloc] initWithData:ValueData encoding:NSUTF8StringEncoding];
+			if (ValueStr)
+			{
+				Info.TxtRecord.Add(FString(Key), FString(ValueStr));
+				[ValueStr release];
+			}
+		}
+	}
+
+	UE_LOG(LogNetworkServiceDiscovery, Log, TEXT("Apple: Resolved service '%s' -> %s:%d"), *Info.ServiceName, *Info.Address, Info.Port);
+
+	// Remove from resolving list
+	[ResolvingServices removeObject:Service];
+	Service.delegate = nil;
+
+	AsyncTask(ENamedThreads::GameThread, [WeakOwner = Owner, Info]()
+	{
+		if (WeakOwner)
+		{
+			WeakOwner->HandleServiceResolved(Info);
+		}
+	});
+}
+
+- (void)netService:(NSNetService*)Service didNotResolve:(NSDictionary<NSString*, NSNumber*>*)ErrorDict
+{
+	NSNumber* ErrorCode = ErrorDict[NSNetServicesErrorCode];
+	FString ServiceName = FString(Service.name);
+	FString ErrorMsg = FString::Printf(TEXT("Apple: Failed to resolve '%s' (error %d)"), *ServiceName, ErrorCode ? [ErrorCode intValue] : -1);
+
+	UE_LOG(LogNetworkServiceDiscovery, Warning, TEXT("%s"), *ErrorMsg);
+
+	[ResolvingServices removeObject:Service];
+	Service.delegate = nil;
+
+	AsyncTask(ENamedThreads::GameThread, [WeakOwner = Owner, ErrorMsg]()
+	{
+		if (WeakOwner)
+		{
+			WeakOwner->HandleError(ErrorMsg);
+		}
+	});
+}
+
+@end
+
+// ============================================================================
+//  Objective-C delegate for NSNetService (registration/publishing)
+// ============================================================================
+
+@interface FNSDServiceDelegate : NSObject <NSNetServiceDelegate>
+{
+	FNetworkServiceDiscoveryApple* Owner;
+	NSNetService* PublishedService;
+}
+
+- (instancetype)initWithOwner:(FNetworkServiceDiscoveryApple*)InOwner;
+- (bool)publishServiceWithName:(NSString*)Name type:(NSString*)Type port:(int)Port txtRecord:(NSDictionary<NSString*, NSData*>*)TxtRecord;
+- (void)unpublish;
+
+@end
+
+@implementation FNSDServiceDelegate
+
+- (instancetype)initWithOwner:(FNetworkServiceDiscoveryApple*)InOwner
+{
+	self = [super init];
+	if (self)
+	{
+		Owner = InOwner;
+		PublishedService = nil;
+	}
+	return self;
+}
+
+- (void)dealloc
+{
+	[self unpublish];
+	[super dealloc];
+}
+
+- (bool)publishServiceWithName:(NSString*)Name type:(NSString*)Type port:(int)Port txtRecord:(NSDictionary<NSString*, NSData*>*)TxtRecord
+{
+	[self unpublish];
+
+	PublishedService = [[NSNetService alloc] initWithDomain:@"local." type:Type name:Name port:Port];
+	if (!PublishedService)
+	{
+		return false;
+	}
+
+	PublishedService.delegate = self;
+
+	if (TxtRecord && TxtRecord.count > 0)
+	{
+		NSData* TxtData = [NSNetService dataFromTXTRecordDictionary:TxtRecord];
+		[PublishedService setTXTRecordData:TxtData];
+	}
+
+	[PublishedService publish];
+	return true;
+}
+
+- (void)unpublish
+{
+	if (PublishedService)
+	{
+		[PublishedService stop];
+		PublishedService.delegate = nil;
+		[PublishedService release];
+		PublishedService = nil;
+	}
+}
+
+// --- NSNetServiceDelegate ---
+
+- (void)netServiceDidPublish:(NSNetService*)Service
+{
+	FNetworkServiceInfo Info;
+	Info.ServiceName = FString(Service.name);
+	Info.ServiceType = FString(Service.type);
+	Info.Port = (int32)Service.port;
+	Info.bIsResolved = true;
+
+	UE_LOG(LogNetworkServiceDiscovery, Log, TEXT("Apple: Published service '%s' on port %d"), *Info.ServiceName, Info.Port);
+
+	AsyncTask(ENamedThreads::GameThread, [WeakOwner = Owner, Info]()
+	{
+		if (WeakOwner)
+		{
+			WeakOwner->HandleServiceRegistered(Info);
+		}
+	});
+}
+
+- (void)netService:(NSNetService*)Service didNotPublish:(NSDictionary<NSString*, NSNumber*>*)ErrorDict
+{
+	NSNumber* ErrorCode = ErrorDict[NSNetServicesErrorCode];
+	FString ErrorMsg = FString::Printf(TEXT("Apple: Failed to publish service (error %d)"), ErrorCode ? [ErrorCode intValue] : -1);
+
+	UE_LOG(LogNetworkServiceDiscovery, Error, TEXT("%s"), *ErrorMsg);
+
+	AsyncTask(ENamedThreads::GameThread, [WeakOwner = Owner, ErrorMsg]()
+	{
+		if (WeakOwner)
+		{
+			WeakOwner->HandleError(ErrorMsg);
+		}
+	});
+}
+
+@end
+
+// ============================================================================
+//  C++ implementation
+// ============================================================================
+
+FNetworkServiceDiscoveryApple::FNetworkServiceDiscoveryApple()
+	: BrowserDelegate(nil)
+	, bIsDiscovering(false)
+{
+}
+
+FNetworkServiceDiscoveryApple::~FNetworkServiceDiscoveryApple()
+{
+	UnregisterService(FString());
+	StopDiscovery();
+}
+
+// --- Registration ---
+
+bool FNetworkServiceDiscoveryApple::RegisterService(const FString& ServiceName, const FString& ServiceType, int32 Port, const TMap<FString, FString>& TxtRecord)
+{
+	// Unregister any existing service with this name
+	UnregisterService(ServiceName);
+
+	NSString* Name = ServiceName.GetNSString();
+	NSString* Type = ServiceType.GetNSString();
+
+	// Build TXT record dictionary
+	NSMutableDictionary<NSString*, NSData*>* TxtDict = nil;
+	if (TxtRecord.Num() > 0)
+	{
+		TxtDict = [NSMutableDictionary dictionaryWithCapacity:TxtRecord.Num()];
+		for (const auto& Pair : TxtRecord)
+		{
+			NSString* Key = Pair.Key.GetNSString();
+			NSData* Value = [Pair.Value.GetNSString() dataUsingEncoding:NSUTF8StringEncoding];
+			[TxtDict setObject:Value forKey:Key];
+		}
+	}
+
+	FNSDServiceDelegate* Delegate = [[FNSDServiceDelegate alloc] initWithOwner:this];
+	bool bSuccess = [Delegate publishServiceWithName:Name type:Type port:Port txtRecord:TxtDict];
+	if (bSuccess)
+	{
+		PublishDelegates.Add(ServiceName, Delegate);
+	}
+	else
+	{
+		[Delegate release];
+	}
+
+	return bSuccess;
+}
+
+void FNetworkServiceDiscoveryApple::UnregisterService(const FString& ServiceName)
+{
+	if (ServiceName.IsEmpty())
+	{
+		// Unregister all
+		for (auto& Pair : PublishDelegates)
+		{
+			[Pair.Value unpublish];
+			[Pair.Value release];
+		}
+		PublishDelegates.Empty();
+	}
+	else
+	{
+		FNSDServiceDelegate** Found = PublishDelegates.Find(ServiceName);
+		if (Found)
+		{
+			[*Found unpublish];
+			[*Found release];
+			PublishDelegates.Remove(ServiceName);
+		}
+	}
+}
+
+bool FNetworkServiceDiscoveryApple::IsServiceRegistered(const FString& ServiceName) const
+{
+	if (ServiceName.IsEmpty())
+	{
+		return PublishDelegates.Num() > 0;
+	}
+	return PublishDelegates.Contains(ServiceName);
+}
+
+// --- Discovery ---
+
+bool FNetworkServiceDiscoveryApple::StartDiscovery(const FString& ServiceType)
+{
+	StopDiscovery();
+
+	BrowserDelegate = [[FNSDServiceBrowserDelegate alloc] initWithOwner:this];
+	[BrowserDelegate startBrowsingForType:ServiceType.GetNSString()];
+	bIsDiscovering = true;
+
+	return true;
+}
+
+void FNetworkServiceDiscoveryApple::StopDiscovery()
+{
+	if (BrowserDelegate)
+	{
+		[BrowserDelegate stopBrowsing];
+		[BrowserDelegate release];
+		BrowserDelegate = nil;
+	}
+	bIsDiscovering = false;
+
+	FScopeLock Lock(&ServicesLock);
+	DiscoveredServices.Empty();
+}
+
+bool FNetworkServiceDiscoveryApple::IsDiscovering() const
+{
+	return bIsDiscovering;
+}
+
+void FNetworkServiceDiscoveryApple::ResolveService(const FNetworkServiceInfo& Service)
+{
+	if (BrowserDelegate)
+	{
+		[BrowserDelegate resolveService:Service.ServiceName.GetNSString() ofType:Service.ServiceType.GetNSString()];
+	}
+}
+
+TArray<FNetworkServiceInfo> FNetworkServiceDiscoveryApple::GetDiscoveredServices() const
+{
+	FScopeLock Lock(&ServicesLock);
+	return DiscoveredServices;
+}
+
+// --- Callbacks from Objective-C delegates (called on game thread) ---
+
+void FNetworkServiceDiscoveryApple::HandleServiceFound(const FNetworkServiceInfo& Service)
+{
+	{
+		FScopeLock Lock(&ServicesLock);
+		DiscoveredServices.Add(Service);
+	}
+	OnServiceFoundDelegate.Broadcast(Service);
+}
+
+void FNetworkServiceDiscoveryApple::HandleServiceLost(const FNetworkServiceInfo& Service)
+{
+	{
+		FScopeLock Lock(&ServicesLock);
+		DiscoveredServices.RemoveAll([&Service](const FNetworkServiceInfo& Existing)
+		{
+			return Existing.ServiceName == Service.ServiceName && Existing.ServiceType == Service.ServiceType;
+		});
+	}
+	OnServiceLostDelegate.Broadcast(Service);
+}
+
+void FNetworkServiceDiscoveryApple::HandleServiceResolved(const FNetworkServiceInfo& Service)
+{
+	{
+		FScopeLock Lock(&ServicesLock);
+		// Update the existing entry with resolved info
+		for (FNetworkServiceInfo& Existing : DiscoveredServices)
+		{
+			if (Existing.ServiceName == Service.ServiceName && Existing.ServiceType == Service.ServiceType)
+			{
+				Existing = Service;
+				break;
+			}
+		}
+	}
+	OnServiceResolvedDelegate.Broadcast(Service);
+}
+
+void FNetworkServiceDiscoveryApple::HandleServiceRegistered(const FNetworkServiceInfo& Service)
+{
+	OnServiceRegisteredDelegate.Broadcast(Service);
+}
+
+void FNetworkServiceDiscoveryApple::HandleError(const FString& ErrorMessage)
+{
+	OnDiscoveryErrorDelegate.Broadcast(ErrorMessage);
+}
diff --git a/Engine/Plugins/Experimental/NetworkServiceDiscovery/Source/NetworkServiceDiscovery/Private/INetworkServiceDiscoveryPlatform.h b/Engine/Plugins/Experimental/NetworkServiceDiscovery/Source/NetworkServiceDiscovery/Private/INetworkServiceDiscoveryPlatform.h
new file mode 100644
--- /dev/null
+++ b/Engine/Plugins/Experimental/NetworkServiceDiscovery/Source/NetworkServiceDiscovery/Private/INetworkServiceDiscoveryPlatform.h
@@ -0,0 +1,34 @@
+// Copyright Epic Games, Inc. All Rights Reserved.
+
+#pragma once
+
+#include "NetworkServiceDiscoveryTypes.h"
+
+/**
+ * Internal platform abstraction for network service discovery.
+ * Each platform provides a concrete implementation of this interface.
+ */
+class INetworkServiceDiscoveryPlatform
+{
+public:
+	virtual ~INetworkServiceDiscoveryPlatform() = default;
+
+	// --- Registration ---
+	virtual bool RegisterService(const FString& ServiceName, const FString& ServiceType, int32 Port, const TMap<FString, FString>& TxtRecord) = 0;
+	virtual void UnregisterService(const FString& ServiceName) = 0;
+	virtual bool IsServiceRegistered(const FString& ServiceName) const = 0;
+
+	// --- Discovery ---
+	virtual bool StartDiscovery(const FString& ServiceType) = 0;
+	virtual void StopDiscovery() = 0;
+	virtual bool IsDiscovering() const = 0;
+	virtual void ResolveService(const FNetworkServiceInfo& Service) = 0;
+	virtual TArray<FNetworkServiceInfo> GetDiscoveredServices() const = 0;
+
+	// --- Delegates ---
+	FOnServiceFound OnServiceFoundDelegate;
+	FOnServiceLost OnServiceLostDelegate;
+	FOnServiceResolved OnServiceResolvedDelegate;
+	FOnServiceRegistered OnServiceRegisteredDelegate;
+	FOnDiscoveryError OnDiscoveryErrorDelegate;
+};
diff --git a/Engine/Plugins/Experimental/NetworkServiceDiscovery/Source/NetworkServiceDiscovery/Private/NetworkServiceDiscoveryModule.cpp b/Engine/Plugins/Experimental/NetworkServiceDiscovery/Source/NetworkServiceDiscovery/Private/NetworkServiceDiscoveryModule.cpp
new file mode 100644
--- /dev/null
+++ b/Engine/Plugins/Experimental/NetworkServiceDiscovery/Source/NetworkServiceDiscovery/Private/NetworkServiceDiscoveryModule.cpp
@@ -0,0 +1,152 @@
+// Copyright Epic Games, Inc. All Rights Reserved.
+
+#include "NetworkServiceDiscoveryModule.h"
+#include "INetworkServiceDiscoveryPlatform.h"
+#include "Modules/ModuleManager.h"
+
+// Platform headers
+#if PLATFORM_IOS || PLATFORM_MAC
+#include "Apple/NetworkServiceDiscoveryApple.h"
+#elif PLATFORM_ANDROID
+#include "Android/NetworkServiceDiscoveryAndroid.h"
+#elif WITH_WINDOWS_DNSSD
+#include "Windows/NetworkServiceDiscoveryWindows.h"
+#include "Null/NetworkServiceDiscoveryNull.h"
+#else
+#include "Null/NetworkServiceDiscoveryNull.h"
+#endif
+
+DEFINE_LOG_CATEGORY(LogNetworkServiceDiscovery);
+
+void FNetworkServiceDiscoveryModule::StartupModule()
+{
+	PlatformImpl = CreatePlatformImpl();
+	UE_LOG(LogNetworkServiceDiscovery, Log, TEXT("NetworkServiceDiscovery module started."));
+}
+
+void FNetworkServiceDiscoveryModule::ShutdownModule()
+{
+	if (PlatformImpl)
+	{
+		PlatformImpl->UnregisterService(FString());
+		PlatformImpl->StopDiscovery();
+		PlatformImpl.Reset();
+	}
+	UE_LOG(LogNetworkServiceDiscovery, Log, TEXT("NetworkServiceDiscovery module shut down."));
+}
+
+TUniquePtr<INetworkServiceDiscoveryPlatform> FNetworkServiceDiscoveryModule::CreatePlatformImpl()
+{
+#if PLATFORM_IOS || PLATFORM_MAC
+	UE_LOG(LogNetworkServiceDiscovery, Log, TEXT("Using Apple (NSNetService) implementation."));
+	return MakeUnique<FNetworkServiceDiscoveryApple>();
+#elif PLATFORM_ANDROID
+	UE_LOG(LogNetworkServiceDiscovery, Log, TEXT("Using Android (NsdManager) implementation."));
+	return MakeUnique<FNetworkServiceDiscoveryAndroid>();
+#elif WITH_WINDOWS_DNSSD
+	if (FNetworkServiceDiscoveryWindows::LoadDnsApi())
+	{
+		UE_LOG(LogNetworkServiceDiscovery, Log, TEXT("Using Windows (native DNS-SD) implementation."));
+		return MakeUnique<FNetworkServiceDiscoveryWindows>();
+	}
+	else
+	{
+		UE_LOG(LogNetworkServiceDiscovery, Warning, TEXT("Windows DNS-SD APIs not available, falling back to Null implementation."));
+		return MakeUnique<FNetworkServiceDiscoveryNull>();
+	}
+#else
+	UE_LOG(LogNetworkServiceDiscovery, Warning, TEXT("Using Null implementation - mDNS not supported on this platform."));
+	return MakeUnique<FNetworkServiceDiscoveryNull>();
+#endif
+}
+
+// --- Registration ---
+
+bool FNetworkServiceDiscoveryModule::RegisterService(const FString& ServiceName, const FString& ServiceType, int32 Port, const TMap<FString, FString>& TxtRecord)
+{
+	check(PlatformImpl);
+	UE_LOG(LogNetworkServiceDiscovery, Log, TEXT("RegisterService: Name='%s' Type='%s' Port=%d"), *ServiceName, *ServiceType, Port);
+	return PlatformImpl->RegisterService(ServiceName, ServiceType, Port, TxtRecord);
+}
+
+void FNetworkServiceDiscoveryModule::UnregisterService(const FString& ServiceName)
+{
+	check(PlatformImpl);
+	UE_LOG(LogNetworkServiceDiscovery, Log, TEXT("UnregisterService: Name='%s'"), ServiceName.IsEmpty() ? TEXT("(all)") : *ServiceName);
+	PlatformImpl->UnregisterService(ServiceName);
+}
+
+bool FNetworkServiceDiscoveryModule::IsServiceRegistered(const FString& ServiceName) const
+{
+	check(PlatformImpl);
+	return PlatformImpl->IsServiceRegistered(ServiceName);
+}
+
+// --- Discovery ---
+
+bool FNetworkServiceDiscoveryModule::StartDiscovery(const FString& ServiceType)
+{
+	check(PlatformImpl);
+	UE_LOG(LogNetworkServiceDiscovery, Log, TEXT("StartDiscovery: Type='%s'"), *ServiceType);
+	return PlatformImpl->StartDiscovery(ServiceType);
+}
+
+void FNetworkServiceDiscoveryModule::StopDiscovery()
+{
+	check(PlatformImpl);
+	UE_LOG(LogNetworkServiceDiscovery, Log, TEXT("StopDiscovery"));
+	PlatformImpl->StopDiscovery();
+}
+
+bool FNetworkServiceDiscoveryModule::IsDiscovering() const
+{
+	check(PlatformImpl);
+	return PlatformImpl->IsDiscovering();
+}
+
+void FNetworkServiceDiscoveryModule::ResolveService(const FNetworkServiceInfo& Service)
+{
+	check(PlatformImpl);
+	UE_LOG(LogNetworkServiceDiscovery, Log, TEXT("ResolveService: Name='%s'"), *Service.ServiceName);
+	PlatformImpl->ResolveService(Service);
+}
+
+TArray<FNetworkServiceInfo> FNetworkServiceDiscoveryModule::GetDiscoveredServices() const
+{
+	check(PlatformImpl);
+	return PlatformImpl->GetDiscoveredServices();
+}
+
+// --- Delegates ---
+
+FOnServiceFound& FNetworkServiceDiscoveryModule::OnServiceFound()
+{
+	check(PlatformImpl);
+	return PlatformImpl->OnServiceFoundDelegate;
+}
+
+FOnServiceLost& FNetworkServiceDiscoveryModule::OnServiceLost()
+{
+	check(PlatformImpl);
+	return PlatformImpl->OnServiceLostDelegate;
+}
+
+FOnServiceResolved& FNetworkServiceDiscoveryModule::OnServiceResolved()
+{
+	check(PlatformImpl);
+	return PlatformImpl->OnServiceResolvedDelegate;
+}
+
+FOnServiceRegistered& FNetworkServiceDiscoveryModule::OnServiceRegistered()
+{
+	check(PlatformImpl);
+	return PlatformImpl->OnServiceRegisteredDelegate;
+}
+
+FOnDiscoveryError& FNetworkServiceDiscoveryModule::OnDiscoveryError()
+{
+	check(PlatformImpl);
+	return PlatformImpl->OnDiscoveryErrorDelegate;
+}
+
+IMPLEMENT_MODULE(FNetworkServiceDiscoveryModule, NetworkServiceDiscovery)
diff --git a/Engine/Plugins/Experimental/NetworkServiceDiscovery/Source/NetworkServiceDiscovery/Private/NetworkServiceDiscoveryModule.h b/Engine/Plugins/Experimental/NetworkServiceDiscovery/Source/NetworkServiceDiscovery/Private/NetworkServiceDiscoveryModule.h
new file mode 100644
--- /dev/null
+++ b/Engine/Plugins/Experimental/NetworkServiceDiscovery/Source/NetworkServiceDiscovery/Private/NetworkServiceDiscoveryModule.h
@@ -0,0 +1,45 @@
+// Copyright Epic Games, Inc. All Rights Reserved.
+
+#pragma once
+
+#include "INetworkServiceDiscovery.h"
+
+class INetworkServiceDiscoveryPlatform;
+
+DECLARE_LOG_CATEGORY_EXTERN(LogNetworkServiceDiscovery, Log, All);
+
+class FNetworkServiceDiscoveryModule : public INetworkServiceDiscoveryModule
+{
+public:
+
+	// IModuleInterface
+	virtual void StartupModule() override;
+	virtual void ShutdownModule() override;
+
+	// INetworkServiceDiscoveryModule - Registration
+	virtual bool RegisterService(const FString& ServiceName, const FString& ServiceType, int32 Port, const TMap<FString, FString>& TxtRecord) override;
+	virtual void UnregisterService(const FString& ServiceName = FString()) override;
+	virtual bool IsServiceRegistered(const FString& ServiceName = FString()) const override;
+
+	// INetworkServiceDiscoveryModule - Discovery
+	virtual bool StartDiscovery(const FString& ServiceType) override;
+	virtual void StopDiscovery() override;
+	virtual bool IsDiscovering() const override;
+	virtual void ResolveService(const FNetworkServiceInfo& Service) override;
+	virtual TArray<FNetworkServiceInfo> GetDiscoveredServices() const override;
+
+	// INetworkServiceDiscoveryModule - Delegates
+	virtual FOnServiceFound& OnServiceFound() override;
+	virtual FOnServiceLost& OnServiceLost() override;
+	virtual FOnServiceResolved& OnServiceResolved() override;
+	virtual FOnServiceRegistered& OnServiceRegistered() override;
+	virtual FOnDiscoveryError& OnDiscoveryError() override;
+
+private:
+
+	/** Platform-specific implementation */
+	TUniquePtr<INetworkServiceDiscoveryPlatform> PlatformImpl;
+
+	/** Create the appropriate platform implementation */
+	static TUniquePtr<INetworkServiceDiscoveryPlatform> CreatePlatformImpl();
+};
diff --git a/Engine/Plugins/Experimental/NetworkServiceDiscovery/Source/NetworkServiceDiscovery/Private/Null/NetworkServiceDiscoveryNull.h b/Engine/Plugins/Experimental/NetworkServiceDiscovery/Source/NetworkServiceDiscovery/Private/Null/NetworkServiceDiscoveryNull.h
new file mode 100644
--- /dev/null
+++ b/Engine/Plugins/Experimental/NetworkServiceDiscovery/Source/NetworkServiceDiscovery/Private/Null/NetworkServiceDiscoveryNull.h
@@ -0,0 +1,57 @@
+// Copyright Epic Games, Inc. All Rights Reserved.
+
+#pragma once
+
+#include "INetworkServiceDiscoveryPlatform.h"
+#include "NetworkServiceDiscoveryModule.h"
+
+/**
+ * Null (no-op) implementation for unsupported platforms.
+ * All operations log a warning and return failure.
+ */
+class FNetworkServiceDiscoveryNull : public INetworkServiceDiscoveryPlatform
+{
+public:
+
+	virtual bool RegisterService(const FString& ServiceName, const FString& ServiceType, int32 Port, const TMap<FString, FString>& TxtRecord) override
+	{
+		LogUnsupported();
+		return false;
+	}
+
+	virtual void UnregisterService(const FString& ServiceName) override {}
+
+	virtual bool IsServiceRegistered(const FString& ServiceName) const override { return false; }
+
+	virtual bool StartDiscovery(const FString& ServiceType) override
+	{
+		LogUnsupported();
+		return false;
+	}
+
+	virtual void StopDiscovery() override {}
+
+	virtual bool IsDiscovering() const override { return false; }
+
+	virtual void ResolveService(const FNetworkServiceInfo& Service) override
+	{
+		LogUnsupported();
+	}
+
+	virtual TArray<FNetworkServiceInfo> GetDiscoveredServices() const override
+	{
+		return TArray<FNetworkServiceInfo>();
+	}
+
+private:
+
+	void LogUnsupported()
+	{
+		static bool bWarned = false;
+		if (!bWarned)
+		{
+			UE_LOG(LogNetworkServiceDiscovery, Warning, TEXT("mDNS/DNS-SD is not supported on this platform."));
+			bWarned = true;
+		}
+	}
+};
diff --git a/Engine/Plugins/Experimental/NetworkServiceDiscovery/Source/NetworkServiceDiscovery/Private/Windows/NetworkServiceDiscoveryWindows.cpp b/Engine/Plugins/Experimental/NetworkServiceDiscovery/Source/NetworkServiceDiscovery/Private/Windows/NetworkServiceDiscoveryWindows.cpp
new file mode 100644
--- /dev/null
+++ b/Engine/Plugins/Experimental/NetworkServiceDiscovery/Source/NetworkServiceDiscovery/Private/Windows/NetworkServiceDiscoveryWindows.cpp
@@ -0,0 +1,811 @@
+// Copyright Epic Games, Inc. All Rights Reserved.
+
+#include "Windows/NetworkServiceDiscoveryWindows.h"
+#include "NetworkServiceDiscoveryModule.h"
+#include "Async/Async.h"
+
+#if WITH_WINDOWS_DNSSD
+
+#include "Windows/AllowWindowsPlatformTypes.h"
+#include <windns.h>
+#include <winsock2.h>
+#include <ws2tcpip.h>
+#include "Windows/HideWindowsPlatformTypes.h"
+
+// ============================================================================
+//  Dynamic loading of dnsapi.dll
+// ============================================================================
+
+// Function pointer typedefs matching the dnsapi.dll exports we use
+typedef PDNS_SERVICE_INSTANCE (WINAPI* FnDnsServiceConstructInstance)(
+	PCWSTR pServiceName, PCWSTR pHostName, PIP4_ADDRESS pIp4, PIP6_ADDRESS pIp6,
+	WORD wPort, WORD wPriority, WORD wWeight, DWORD dwPropertiesCount,
+	PCWSTR* keys, PCWSTR* values);
+typedef VOID (WINAPI* FnDnsServiceFreeInstance)(PDNS_SERVICE_INSTANCE pInstance);
+typedef DWORD (WINAPI* FnDnsServiceRegister)(PDNS_SERVICE_REGISTER_REQUEST pRequest, PDNS_SERVICE_CANCEL pCancel);
+typedef DWORD (WINAPI* FnDnsServiceDeRegister)(PDNS_SERVICE_REGISTER_REQUEST pRequest, PDNS_SERVICE_CANCEL pCancel);
+typedef DWORD (WINAPI* FnDnsServiceRegisterCancel)(PDNS_SERVICE_CANCEL pCancelHandle);
+typedef DWORD (WINAPI* FnDnsServiceBrowse)(PDNS_SERVICE_BROWSE_REQUEST pRequest, PDNS_SERVICE_CANCEL pCancel);
+typedef DWORD (WINAPI* FnDnsServiceBrowseCancel)(PDNS_SERVICE_CANCEL pCancelHandle);
+typedef DWORD (WINAPI* FnDnsServiceResolve)(PDNS_SERVICE_RESOLVE_REQUEST pRequest, PDNS_SERVICE_CANCEL pCancel);
+typedef DWORD (WINAPI* FnDnsServiceResolveCancel)(PDNS_SERVICE_CANCEL pCancelHandle);
+typedef VOID (WINAPI* FnDnsRecordListFree)(PDNS_RECORD pRecordList, DNS_FREE_TYPE FreeType);
+
+// Static function pointers — populated by LoadDnsApi()
+static FnDnsServiceConstructInstance  pDnsServiceConstructInstance  = nullptr;
+static FnDnsServiceFreeInstance       pDnsServiceFreeInstance       = nullptr;
+static FnDnsServiceRegister           pDnsServiceRegister           = nullptr;
+static FnDnsServiceDeRegister         pDnsServiceDeRegister         = nullptr;
+static FnDnsServiceRegisterCancel     pDnsServiceRegisterCancel     = nullptr;
+static FnDnsServiceBrowse             pDnsServiceBrowse             = nullptr;
+static FnDnsServiceBrowseCancel       pDnsServiceBrowseCancel       = nullptr;
+static FnDnsServiceResolve            pDnsServiceResolve            = nullptr;
+static FnDnsServiceResolveCancel      pDnsServiceResolveCancel      = nullptr;
+static FnDnsRecordListFree            pDnsRecordListFree            = nullptr;
+
+bool FNetworkServiceDiscoveryWindows::bDnsApiLoaded = false;
+
+bool FNetworkServiceDiscoveryWindows::LoadDnsApi()
+{
+	static bool bAttempted = false;
+	if (bAttempted)
+	{
+		return bDnsApiLoaded;
+	}
+	bAttempted = true;
+
+	HMODULE hDnsApi = LoadLibraryW(L"dnsapi.dll");
+	if (!hDnsApi)
+	{
+		UE_LOG(LogNetworkServiceDiscovery, Warning, TEXT("Windows: Failed to load dnsapi.dll — DNS-SD will not be available"));
+		return false;
+	}
+
+	#define NSD_LOAD_FUNC(Name) \
+		__pragma(warning(suppress: 4191)) /* FARPROC to typed function pointer cast is expected */ \
+		p##Name = reinterpret_cast<Fn##Name>(GetProcAddress(hDnsApi, #Name)); \
+		if (!p##Name) \
+		{ \
+			UE_LOG(LogNetworkServiceDiscovery, Warning, TEXT("Windows: dnsapi.dll missing export '%s' — DNS-SD will not be available"), TEXT(#Name)); \
+			FreeLibrary(hDnsApi); \
+			return false; \
+		}
+
+	NSD_LOAD_FUNC(DnsServiceConstructInstance);
+	NSD_LOAD_FUNC(DnsServiceFreeInstance);
+	NSD_LOAD_FUNC(DnsServiceRegister);
+	NSD_LOAD_FUNC(DnsServiceDeRegister);
+	NSD_LOAD_FUNC(DnsServiceRegisterCancel);
+	NSD_LOAD_FUNC(DnsServiceBrowse);
+	NSD_LOAD_FUNC(DnsServiceBrowseCancel);
+	NSD_LOAD_FUNC(DnsServiceResolve);
+	NSD_LOAD_FUNC(DnsServiceResolveCancel);
+	NSD_LOAD_FUNC(DnsRecordListFree);
+
+	#undef NSD_LOAD_FUNC
+
+	// Intentionally never FreeLibrary — the DLL must stay loaded for the
+	// lifetime of the process since async callbacks reference its code.
+	bDnsApiLoaded = true;
+	UE_LOG(LogNetworkServiceDiscovery, Log, TEXT("Windows: dnsapi.dll loaded successfully"));
+	return true;
+}
+
+/** Re-browse interval in seconds for detecting service removals */
+static constexpr float ReBrowseIntervalSeconds = 5.0f;
+
+// ============================================================================
+//  Helpers
+// ============================================================================
+
+static FString BuildInstanceName(const FString& ServiceName, const FString& ServiceType)
+{
+	FString CleanType = ServiceType;
+	CleanType.RemoveFromEnd(TEXT("."));
+	return FString::Printf(TEXT("%s.%s.local"), *ServiceName, *CleanType);
+}
+
+// ============================================================================
+//  Callback context structs
+// ============================================================================
+
+/** Context for register/deregister callbacks — captures the alive flag at call-site time.
+ *  For register calls, ServiceName + Generation identify the registration so the callback
+ *  can store the actual DNS_SERVICE_INSTANCE (which may have a different name due to mDNS
+ *  conflict resolution). Stale callbacks (from a cancelled/replaced registration) are
+ *  detected by generation mismatch and ignored.
+ *  For deregister calls, ServiceInstanceToFree holds the DNS_SERVICE_INSTANCE* to free
+ *  once the async deregistration completes (must not be freed before the callback). */
+struct FRegisterContext
+{
+	FNetworkServiceDiscoveryWindows* Self;
+	TSharedPtr<TAtomic<bool>> AliveFlag;
+	FString ServiceName;
+	uint64 Generation = 0;
+	PDNS_SERVICE_INSTANCE ServiceInstanceToFree = nullptr;
+	bool bIsDeregister = false;
+};
+
+/** Context for browse callbacks — captures the alive flag at call-site time.
+ *  Ref-counted via TSharedPtr: the caller holds one ref (BrowseCallbackContext member)
+ *  and a heap-allocated TSharedPtr copy is passed as the PVOID context to the Windows API.
+ *  The callback extracts and copies the shared pointer, then deletes the heap wrapper.
+ *  This ensures the context stays alive regardless of whether DnsServiceBrowseCancel
+ *  triggers a final callback or not — whichever side releases last frees the memory. */
+struct FBrowseContext
+{
+	FNetworkServiceDiscoveryWindows* Self;
+	TSharedPtr<TAtomic<bool>> AliveFlag;
+};
+
+// ============================================================================
+//  File-static callbacks
+// ============================================================================
+
+static void WINAPI NSD_RegisterCallback(DWORD Status, PVOID Context, PDNS_SERVICE_INSTANCE pInstance)
+{
+	FRegisterContext* Ctx = static_cast<FRegisterContext*>(Context);
+	FNetworkServiceDiscoveryWindows* Self = Ctx->Self;
+	TSharedPtr<TAtomic<bool>> AliveFlag = Ctx->AliveFlag;
+	FString ServiceName = MoveTemp(Ctx->ServiceName);
+	uint64 Generation = Ctx->Generation;
+
+	// For deregister calls, free the service instance now that deregistration is complete.
+	// This must be deferred until the callback because DnsServiceDeRegister is async —
+	// freeing the instance before deregistration completes leaves a ghost service on the
+	// network, causing name conflicts (e.g. "MyService (2)") on re-registration.
+	bool bIsDeregister = Ctx->bIsDeregister;
+	if (Ctx->ServiceInstanceToFree)
+	{
+		pDnsServiceFreeInstance(Ctx->ServiceInstanceToFree);
+	}
+	delete Ctx;
+
+	if (bIsDeregister)
+	{
+		// Deregister completion — the system may provide a pInstance copy, free it if present.
+		// Don't store it or broadcast OnServiceRegistered.
+		if (pInstance)
+		{
+			pDnsServiceFreeInstance(pInstance);
+		}
+		return;
+	}
+
+	if (Status == ERROR_SUCCESS)
+	{
+		FNetworkServiceInfo Info;
+		if (pInstance && pInstance->pszInstanceName)
+		{
+			Info.ServiceName = FString(pInstance->pszInstanceName);
+			Info.Port = pInstance->wPort;
+		}
+		Info.bIsResolved = true;
+
+		UE_LOG(LogNetworkServiceDiscovery, Log, TEXT("Windows: Service registered '%s'"), *Info.ServiceName);
+
+		// Store the callback's DNS_SERVICE_INSTANCE as ActualInstance — it has the real
+		// registered name (which may differ from what we requested due to mDNS conflict
+		// resolution). DnsServiceDeRegister needs this actual-name instance to work.
+		// Generation check prevents stale callbacks from overwriting a newer registration.
+		// We own pInstance and must free it if we can't store it.
+		AsyncTask(ENamedThreads::GameThread, [Self, AliveFlag, ServiceName, Generation, Info, pInstance]()
+		{
+			if (AliveFlag->Load())
+			{
+				bool bStored = false;
+				if (!ServiceName.IsEmpty())
+				{
+					auto* Reg = Self->Registrations.Find(ServiceName);
+					if (Reg && Reg->Generation == Generation)
+					{
+						// This callback matches the current registration — store the actual instance
+						if (Reg->OriginalInstance)
+						{
+							pDnsServiceFreeInstance(static_cast<PDNS_SERVICE_INSTANCE>(Reg->OriginalInstance));
+							Reg->OriginalInstance = nullptr;
+						}
+						Reg->ActualInstance = pInstance;
+						bStored = true;
+					}
+				}
+				if (!bStored && pInstance)
+				{
+					// Stale callback or registration already removed — free the instance
+					pDnsServiceFreeInstance(pInstance);
+				}
+				Self->OnServiceRegisteredDelegate.Broadcast(Info);
+			}
+			else if (pInstance)
+			{
+				pDnsServiceFreeInstance(pInstance);
+			}
+		});
+	}
+	else
+	{
+		FString ErrorMsg = FString::Printf(TEXT("Windows: Registration failed with error %u"), Status);
+		UE_LOG(LogNetworkServiceDiscovery, Error, TEXT("%s"), *ErrorMsg);
+
+		AsyncTask(ENamedThreads::GameThread, [Self, AliveFlag, ErrorMsg]()
+		{
+			if (AliveFlag->Load()) { Self->OnDiscoveryErrorDelegate.Broadcast(ErrorMsg); }
+		});
+	}
+}
+
+static void WINAPI NSD_BrowseCallback(DWORD Status, PVOID Context, PDNS_RECORD pDnsRecord)
+{
+	// Extract the shared pointer from the heap wrapper — this bumps the ref count so the
+	// FBrowseContext stays alive for the duration of this callback even if the caller
+	// (StopDiscovery/OnReBrowseTick) has already released its ref.
+	TSharedPtr<FBrowseContext>* pSharedCtx = static_cast<TSharedPtr<FBrowseContext>*>(Context);
+	TSharedPtr<FBrowseContext> Ctx = *pSharedCtx;
+	delete pSharedCtx;
+
+	FNetworkServiceDiscoveryWindows* Self = Ctx->Self;
+	TSharedPtr<TAtomic<bool>> AliveFlag = Ctx->AliveFlag;
+
+	if (Status != ERROR_SUCCESS || !pDnsRecord)
+	{
+		if (Status != ERROR_SUCCESS && Status != ERROR_CANCELLED)
+		{
+			FString ErrorMsg = FString::Printf(TEXT("Windows: Browse callback error %u"), Status);
+			UE_LOG(LogNetworkServiceDiscovery, Warning, TEXT("%s"), *ErrorMsg);
+
+			AsyncTask(ENamedThreads::GameThread, [Self, AliveFlag, ErrorMsg]()
+			{
+				if (AliveFlag->Load()) { Self->OnDiscoveryErrorDelegate.Broadcast(ErrorMsg); }
+			});
+		}
+		return;
+	}
+
+	// Walk the DNS record list — PTR records point to service instance names
+	TArray<FNetworkServiceInfo> FoundServices;
+	for (PDNS_RECORD pRecord = pDnsRecord; pRecord != nullptr; pRecord = pRecord->pNext)
+	{
+		if (pRecord->wType == DNS_TYPE_PTR && pRecord->Data.PTR.pNameHost)
+		{
+			FNetworkServiceInfo Info;
+			Info.ServiceName = FString(pRecord->Data.PTR.pNameHost);
+			FoundServices.Add(Info);
+
+			UE_LOG(LogNetworkServiceDiscovery, Verbose, TEXT("Windows: Browse found '%s'"), *Info.ServiceName);
+		}
+	}
+
+	pDnsRecordListFree(pDnsRecord, DnsFreeRecordList);
+
+	AsyncTask(ENamedThreads::GameThread, [Self, AliveFlag, FoundServices]()
+	{
+		if (!AliveFlag->Load()) { return; }
+		// Track which services were seen in this cycle for removal detection
+		for (const FNetworkServiceInfo& Info : FoundServices)
+		{
+			Self->CurrentBrowseCycleServices.Add(Info.ServiceName);
+
+			bool bAlreadyKnown = false;
+			{
+				FScopeLock Lock(&Self->ServicesLock);
+				for (const FNetworkServiceInfo& Existing : Self->DiscoveredServices)
+				{
+					if (Existing.ServiceName == Info.ServiceName)
+					{
+						bAlreadyKnown = true;
+						break;
+					}
+				}
+				if (!bAlreadyKnown)
+				{
+					Self->DiscoveredServices.Add(Info);
+				}
+			}
+
+			if (!bAlreadyKnown)
+			{
+				Self->OnServiceFoundDelegate.Broadcast(Info);
+			}
+		}
+	});
+}
+
+/** Context for resolve callbacks — carries the owner and the heap-allocated query name */
+struct FResolveContext
+{
+	FNetworkServiceDiscoveryWindows* Self;
+	TSharedPtr<TAtomic<bool>> AliveFlag;
+	WCHAR* QueryNameBuf;
+	PDNS_SERVICE_CANCEL CancelHandle;
+};
+
+static void WINAPI NSD_ResolveCallback(DWORD Status, PVOID Context, PDNS_SERVICE_INSTANCE pInstance)
+{
+	FResolveContext* Ctx = static_cast<FResolveContext*>(Context);
+	FNetworkServiceDiscoveryWindows* Self = Ctx->Self;
+	TSharedPtr<TAtomic<bool>> AliveFlag = Ctx->AliveFlag;
+
+	// Always clean up context resources regardless of object lifetime
+	delete[] Ctx->QueryNameBuf;
+	if (AliveFlag->Load())
+	{
+		FScopeLock Lock(&Self->ResolveLock);
+		Self->PendingResolveCancels.Remove(Ctx->CancelHandle);
+	}
+	delete Ctx->CancelHandle;
+	delete Ctx;
+
+	if (Status != ERROR_SUCCESS || !pInstance)
+	{
+		if (Status != ERROR_SUCCESS && Status != ERROR_CANCELLED)
+		{
+			FString ErrorMsg = FString::Printf(TEXT("Windows: Resolve failed with error %u"), Status);
+			UE_LOG(LogNetworkServiceDiscovery, Warning, TEXT("%s"), *ErrorMsg);
+
+			AsyncTask(ENamedThreads::GameThread, [Self, AliveFlag, ErrorMsg]()
+			{
+				if (AliveFlag->Load()) { Self->OnDiscoveryErrorDelegate.Broadcast(ErrorMsg); }
+			});
+		}
+		return;
+	}
+
+	FNetworkServiceInfo Info;
+	Info.ServiceName = pInstance->pszInstanceName ? FString(pInstance->pszInstanceName) : FString();
+	Info.HostName = pInstance->pszHostName ? FString(pInstance->pszHostName) : FString();
+	Info.Port = pInstance->wPort;
+	Info.bIsResolved = true;
+
+	if (pInstance->ip4Address)
+	{
+		IN_ADDR Addr;
+		Addr.S_un.S_addr = *pInstance->ip4Address;
+		WCHAR AddrStr[INET_ADDRSTRLEN];
+		InetNtopW(AF_INET, &Addr, AddrStr, INET_ADDRSTRLEN);
+		Info.Address = FString(AddrStr);
+	}
+
+	for (DWORD i = 0; i < pInstance->dwPropertyCount; i++)
+	{
+		if (pInstance->keys[i])
+		{
+			FString Key = FString(pInstance->keys[i]);
+			FString Value = pInstance->values[i] ? FString(pInstance->values[i]) : FString();
+			Info.TxtRecord.Add(Key, Value);
+		}
+	}
+
+	UE_LOG(LogNetworkServiceDiscovery, Log, TEXT("Windows: Resolved '%s' -> %s:%d"), *Info.ServiceName, *Info.Address, Info.Port);
+
+	AsyncTask(ENamedThreads::GameThread, [Self, AliveFlag, Info]()
+	{
+		if (!AliveFlag->Load()) { return; }
+		{
+			FScopeLock Lock(&Self->ServicesLock);
+			for (FNetworkServiceInfo& Existing : Self->DiscoveredServices)
+			{
+				if (Existing.ServiceName == Info.ServiceName)
+				{
+					Existing = Info;
+					break;
+				}
+			}
+		}
+		Self->OnServiceResolvedDelegate.Broadcast(Info);
+	});
+}
+
+// ============================================================================
+//  Construction / Destruction
+// ============================================================================
+
+FNetworkServiceDiscoveryWindows::FNetworkServiceDiscoveryWindows()
+{
+}
+
+FNetworkServiceDiscoveryWindows::~FNetworkServiceDiscoveryWindows()
+{
+	// Signal all async callbacks that this object is no longer valid
+	bIsAlive->Store(false);
+
+	UnregisterService(FString());
+	StopDiscovery();
+
+	// Cancel any pending resolves — callbacks may still fire but will
+	// see bIsAlive==false and skip accessing this object
+	FScopeLock Lock(&ResolveLock);
+	for (void* Handle : PendingResolveCancels)
+	{
+		pDnsServiceResolveCancel(static_cast<PDNS_SERVICE_CANCEL>(Handle));
+		// Don't delete here — the cancel triggers the callback which handles cleanup
+	}
+	PendingResolveCancels.Empty();
+}
+
+// ============================================================================
+//  Service Registration
+// ============================================================================
+
+bool FNetworkServiceDiscoveryWindows::RegisterService(const FString& ServiceName, const FString& ServiceType, int32 Port, const TMap<FString, FString>& TxtRecord)
+{
+	UnregisterService(ServiceName);
+
+	FString InstanceName = BuildInstanceName(ServiceName, ServiceType);
+	FString HostName = FString::Printf(TEXT("%s.local"), FPlatformProcess::ComputerName());
+
+	// Build TXT record arrays — FStrings own the data, pointer arrays are views for the API
+	TArray<FString> TxtKeyStrings;
+	TArray<FString> TxtValueStrings;
+	for (const auto& Pair : TxtRecord)
+	{
+		TxtKeyStrings.Add(Pair.Key);
+		TxtValueStrings.Add(Pair.Value);
+	}
+
+	TArray<PCWSTR> TxtKeyPtrs;
+	TArray<PCWSTR> TxtValuePtrs;
+	for (int32 i = 0; i < TxtKeyStrings.Num(); i++)
+	{
+		TxtKeyPtrs.Add(*TxtKeyStrings[i]);
+		TxtValuePtrs.Add(*TxtValueStrings[i]);
+	}
+
+	PDNS_SERVICE_INSTANCE pInstance = pDnsServiceConstructInstance(
+		*InstanceName,
+		*HostName,
+		nullptr,    // IPv4 — let the system determine
+		nullptr,    // IPv6
+		(WORD)Port,
+		0,          // priority
+		0,          // weight
+		(DWORD)TxtKeyPtrs.Num(),
+		TxtKeyPtrs.GetData(),
+		TxtValuePtrs.GetData()
+	);
+
+	if (!pInstance)
+	{
+		UE_LOG(LogNetworkServiceDiscovery, Error, TEXT("Windows: DnsServiceConstructInstance failed"));
+		return false;
+	}
+
+	PDNS_SERVICE_CANCEL pCancel = new DNS_SERVICE_CANCEL;
+	FMemory::Memzero(pCancel, sizeof(DNS_SERVICE_CANCEL));
+
+	DNS_SERVICE_REGISTER_REQUEST Request;
+	FMemory::Memzero(&Request, sizeof(Request));
+	Request.Version = DNS_QUERY_REQUEST_VERSION1;
+	Request.InterfaceIndex = 0;
+	Request.pServiceInstance = pInstance;
+	Request.pRegisterCompletionCallback = &NSD_RegisterCallback;
+	Request.unicastEnabled = false;
+
+	++RegistrationGeneration;
+
+	FRegisterContext* RegCtx = new FRegisterContext;
+	RegCtx->Self = this;
+	RegCtx->AliveFlag = bIsAlive;
+	RegCtx->ServiceName = ServiceName;
+	RegCtx->Generation = RegistrationGeneration;
+	Request.pQueryContext = RegCtx;
+
+	DWORD Result = pDnsServiceRegister(&Request, pCancel);
+
+	if (Result != DNS_REQUEST_PENDING && Result != ERROR_SUCCESS)
+	{
+		UE_LOG(LogNetworkServiceDiscovery, Error, TEXT("Windows: DnsServiceRegister failed with error %u"), Result);
+		pDnsServiceFreeInstance(pInstance);
+		delete pCancel;
+		delete RegCtx;
+		return false;
+	}
+
+	FRegistration Reg;
+	Reg.OriginalInstance = pInstance;
+	Reg.ActualInstance = nullptr;
+	Reg.CancelHandle = pCancel;
+	Reg.ServiceName = ServiceName;
+	Reg.Generation = RegistrationGeneration;
+	Registrations.Add(ServiceName, Reg);
+
+	UE_LOG(LogNetworkServiceDiscovery, Log, TEXT("Windows: Registering service '%s' on port %d"), *InstanceName, Port);
+	return true;
+}
+
+void FNetworkServiceDiscoveryWindows::UnregisterService(const FString& ServiceName)
+{
+	auto DeregisterOne = [this](FRegistration& Reg)
+	{
+		if (Reg.ActualInstance)
+		{
+			// Case A: Registration completed — deregister with the actual instance (correct name).
+			// DnsServiceDeRegister matches by instance name, so we must use the instance from
+			// the callback which has the real registered name (possibly renamed by conflict resolution).
+			DNS_SERVICE_REGISTER_REQUEST Request;
+			FMemory::Memzero(&Request, sizeof(Request));
+			Request.Version = DNS_QUERY_REQUEST_VERSION1;
+			Request.InterfaceIndex = 0;
+			Request.pServiceInstance = static_cast<PDNS_SERVICE_INSTANCE>(Reg.ActualInstance);
+			Request.pRegisterCompletionCallback = &NSD_RegisterCallback;
+
+			FRegisterContext* DeregCtx = new FRegisterContext;
+			DeregCtx->Self = this;
+			DeregCtx->AliveFlag = bIsAlive;
+			DeregCtx->ServiceInstanceToFree = static_cast<PDNS_SERVICE_INSTANCE>(Reg.ActualInstance);
+			DeregCtx->bIsDeregister = true;
+			Request.pQueryContext = DeregCtx;
+
+			pDnsServiceDeRegister(&Request, nullptr);
+			// Don't free the instance here — DnsServiceDeRegister is async.
+			// The callback will free it once deregistration actually completes.
+			Reg.ActualInstance = nullptr;
+		}
+		else if (Reg.OriginalInstance)
+		{
+			// Case B: Registration still pending — cancel by handle (no name needed).
+			// DnsServiceRegisterCancel uses the cancel handle to identify the registration,
+			// so it works regardless of what name the system would have assigned.
+			if (Reg.CancelHandle)
+			{
+				pDnsServiceRegisterCancel(static_cast<PDNS_SERVICE_CANCEL>(Reg.CancelHandle));
+			}
+			pDnsServiceFreeInstance(static_cast<PDNS_SERVICE_INSTANCE>(Reg.OriginalInstance));
+			Reg.OriginalInstance = nullptr;
+			// Zero the generation so the pending callback (if it still fires) sees a
+			// mismatch and discards its pInstance instead of storing it.
+			Reg.Generation = 0;
+		}
+		if (Reg.CancelHandle)
+		{
+			delete static_cast<PDNS_SERVICE_CANCEL>(Reg.CancelHandle);
+			Reg.CancelHandle = nullptr;
+		}
+	};
+
+	if (ServiceName.IsEmpty())
+	{
+		for (auto& Pair : Registrations)
+		{
+			DeregisterOne(Pair.Value);
+		}
+		Registrations.Empty();
+	}
+	else
+	{
+		FRegistration* Found = Registrations.Find(ServiceName);
+		if (Found)
+		{
+			DeregisterOne(*Found);
+			Registrations.Remove(ServiceName);
+		}
+	}
+}
+
+bool FNetworkServiceDiscoveryWindows::IsServiceRegistered(const FString& ServiceName) const
+{
+	if (ServiceName.IsEmpty())
+	{
+		return Registrations.Num() > 0;
+	}
+	return Registrations.Contains(ServiceName);
+}
+
+// ============================================================================
+//  Service Discovery (Browse)
+// ============================================================================
+
+bool FNetworkServiceDiscoveryWindows::StartDiscovery(const FString& ServiceType)
+{
+	StopDiscovery();
+
+	BrowseServiceType = ServiceType;
+
+	// Build query name: "_unrealremote._tcp.local" — stored as member to keep alive for async browse
+	FString CleanType = ServiceType;
+	CleanType.RemoveFromEnd(TEXT("."));
+	BrowseQueryName = FString::Printf(TEXT("%s.local"), *CleanType);
+
+	PDNS_SERVICE_CANCEL pCancel = new DNS_SERVICE_CANCEL;
+	FMemory::Memzero(pCancel, sizeof(DNS_SERVICE_CANCEL));
+
+	DNS_SERVICE_BROWSE_REQUEST Request;
+	FMemory::Memzero(&Request, sizeof(Request));
+	Request.Version = DNS_QUERY_REQUEST_VERSION1;
+	Request.InterfaceIndex = 0;
+	Request.QueryName = const_cast<PWSTR>(*BrowseQueryName);
+	Request.pBrowseCallback = &NSD_BrowseCallback;
+
+	TSharedPtr<FBrowseContext> BrowseCtx = MakeShared<FBrowseContext>();
+	BrowseCtx->Self = this;
+	BrowseCtx->AliveFlag = bIsAlive;
+	Request.pQueryContext = new TSharedPtr<FBrowseContext>(BrowseCtx);
+
+	DWORD Result = pDnsServiceBrowse(&Request, pCancel);
+
+	if (Result != DNS_REQUEST_PENDING && Result != ERROR_SUCCESS)
+	{
+		UE_LOG(LogNetworkServiceDiscovery, Error, TEXT("Windows: DnsServiceBrowse failed with error %u"), Result);
+		delete pCancel;
+		delete static_cast<TSharedPtr<FBrowseContext>*>(Request.pQueryContext);
+		return false;
+	}
+
+	BrowseCancelHandle = pCancel;
+	BrowseCallbackContext = BrowseCtx;
+	bIsDiscovering = true;
+
+	// Start periodic re-browse to detect service removals
+	// (Windows DnsServiceBrowse does not notify when services disappear)
+	ReBrowseTickerHandle = FTSTicker::GetCoreTicker().AddTicker(
+		TEXT("NetworkServiceDiscoveryReBrowse"),
+		ReBrowseIntervalSeconds,
+		[this](float DeltaTime) { return OnReBrowseTick(DeltaTime); });
+
+	UE_LOG(LogNetworkServiceDiscovery, Log, TEXT("Windows: Started browsing for '%s'"), *BrowseQueryName);
+	return true;
+}
+
+void FNetworkServiceDiscoveryWindows::StopDiscovery()
+{
+	if (ReBrowseTickerHandle.IsValid())
+	{
+		FTSTicker::GetCoreTicker().RemoveTicker(ReBrowseTickerHandle);
+		ReBrowseTickerHandle.Reset();
+	}
+
+	if (BrowseCancelHandle)
+	{
+		pDnsServiceBrowseCancel(static_cast<PDNS_SERVICE_CANCEL>(BrowseCancelHandle));
+		delete static_cast<PDNS_SERVICE_CANCEL>(BrowseCancelHandle);
+		BrowseCancelHandle = nullptr;
+	}
+	BrowseCallbackContext.Reset();
+	bIsDiscovering = false;
+	BrowseServiceType.Empty();
+	BrowseQueryName.Empty();
+	CurrentBrowseCycleServices.Empty();
+	PreviousBrowseCycleServices.Empty();
+	bFirstBrowseCycleComplete = false;
+
+	FScopeLock Lock(&ServicesLock);
+	DiscoveredServices.Empty();
+}
+
+bool FNetworkServiceDiscoveryWindows::IsDiscovering() const
+{
+	return bIsDiscovering;
+}
+
+bool FNetworkServiceDiscoveryWindows::OnReBrowseTick(float DeltaTime)
+{
+	if (!bIsDiscovering || BrowseServiceType.IsEmpty())
+	{
+		return true; // keep ticking
+	}
+
+	// Swap: the current cycle becomes the previous, start a fresh current set.
+	// The previous set contains all services seen in the last completed browse cycle.
+	PreviousBrowseCycleServices = MoveTemp(CurrentBrowseCycleServices);
+	CurrentBrowseCycleServices.Empty();
+
+	// Only diff after the first full cycle has completed — skip the very first tick
+	// where PreviousBrowseCycleServices would be empty and cause false removals.
+	if (bFirstBrowseCycleComplete)
+	{
+		TArray<FNetworkServiceInfo> LostServices;
+		{
+			FScopeLock Lock(&ServicesLock);
+			for (int32 i = DiscoveredServices.Num() - 1; i >= 0; --i)
+			{
+				if (!PreviousBrowseCycleServices.Contains(DiscoveredServices[i].ServiceName))
+				{
+					LostServices.Add(DiscoveredServices[i]);
+					DiscoveredServices.RemoveAtSwap(i);
+				}
+			}
+		}
+
+		for (const FNetworkServiceInfo& Lost : LostServices)
+		{
+			UE_LOG(LogNetworkServiceDiscovery, Log, TEXT("Windows: Service lost '%s' (not seen in re-browse)"), *Lost.ServiceName);
+			OnServiceLostDelegate.Broadcast(Lost);
+		}
+	}
+	else
+	{
+		bFirstBrowseCycleComplete = true;
+	}
+
+	// Issue a new browse — results will populate CurrentBrowseCycleServices via the callback
+	if (BrowseCancelHandle)
+	{
+		pDnsServiceBrowseCancel(static_cast<PDNS_SERVICE_CANCEL>(BrowseCancelHandle));
+		delete static_cast<PDNS_SERVICE_CANCEL>(BrowseCancelHandle);
+		BrowseCancelHandle = nullptr;
+	}
+	BrowseCallbackContext.Reset();
+
+	PDNS_SERVICE_CANCEL pCancel = new DNS_SERVICE_CANCEL;
+	FMemory::Memzero(pCancel, sizeof(DNS_SERVICE_CANCEL));
+
+	TSharedPtr<FBrowseContext> BrowseCtx = MakeShared<FBrowseContext>();
+	BrowseCtx->Self = this;
+	BrowseCtx->AliveFlag = bIsAlive;
+
+	DNS_SERVICE_BROWSE_REQUEST Request;
+	FMemory::Memzero(&Request, sizeof(Request));
+	Request.Version = DNS_QUERY_REQUEST_VERSION1;
+	Request.InterfaceIndex = 0;
+	Request.QueryName = const_cast<PWSTR>(*BrowseQueryName);
+	Request.pBrowseCallback = &NSD_BrowseCallback;
+	Request.pQueryContext = new TSharedPtr<FBrowseContext>(BrowseCtx);
+
+	DWORD Result = pDnsServiceBrowse(&Request, pCancel);
+	if (Result == DNS_REQUEST_PENDING || Result == ERROR_SUCCESS)
+	{
+		BrowseCancelHandle = pCancel;
+		BrowseCallbackContext = BrowseCtx;
+	}
+	else
+	{
+		delete pCancel;
+		delete static_cast<TSharedPtr<FBrowseContext>*>(Request.pQueryContext);
+	}
+
+	return true; // keep ticking
+}
+
+// ============================================================================
+//  Service Resolution
+// ============================================================================
+
+void FNetworkServiceDiscoveryWindows::ResolveService(const FNetworkServiceInfo& Service)
+{
+	// Allocate the query name on the heap — DnsServiceResolve stores the pointer
+	// for the async operation's lifetime. Freed in the resolve callback via FResolveContext.
+	WCHAR* QueryNameBuf = new WCHAR[Service.ServiceName.Len() + 1];
+	FCString::Strncpy(QueryNameBuf, *Service.ServiceName, Service.ServiceName.Len() + 1);
+
+	PDNS_SERVICE_CANCEL pCancel = new DNS_SERVICE_CANCEL;
+	FMemory::Memzero(pCancel, sizeof(DNS_SERVICE_CANCEL));
+
+	{
+		FScopeLock Lock(&ResolveLock);
+		PendingResolveCancels.Add(pCancel);
+	}
+
+	FResolveContext* Ctx = new FResolveContext;
+	Ctx->Self = this;
+	Ctx->AliveFlag = bIsAlive;
+	Ctx->QueryNameBuf = QueryNameBuf;
+	Ctx->CancelHandle = pCancel;
+
+	DNS_SERVICE_RESOLVE_REQUEST Request;
+	FMemory::Memzero(&Request, sizeof(Request));
+	Request.Version = DNS_QUERY_REQUEST_VERSION1;
+	Request.InterfaceIndex = 0;
+	Request.QueryName = QueryNameBuf;
+	Request.pResolveCompletionCallback = &NSD_ResolveCallback;
+	Request.pQueryContext = Ctx;
+
+	DWORD Result = pDnsServiceResolve(&Request, pCancel);
+
+	if (Result != DNS_REQUEST_PENDING && Result != ERROR_SUCCESS)
+	{
+		UE_LOG(LogNetworkServiceDiscovery, Error, TEXT("Windows: DnsServiceResolve failed with error %u"), Result);
+
+		FScopeLock Lock(&ResolveLock);
+		PendingResolveCancels.Remove(pCancel);
+		delete pCancel;
+		delete[] QueryNameBuf;
+		delete Ctx;
+	}
+}
+
+TArray<FNetworkServiceInfo> FNetworkServiceDiscoveryWindows::GetDiscoveredServices() const
+{
+	FScopeLock Lock(&ServicesLock);
+	return DiscoveredServices;
+}
+
+#endif // WITH_WINDOWS_DNSSD
diff --git a/Engine/Plugins/Experimental/NetworkServiceDiscovery/Source/NetworkServiceDiscovery/Private/Windows/NetworkServiceDiscoveryWindows.h b/Engine/Plugins/Experimental/NetworkServiceDiscovery/Source/NetworkServiceDiscovery/Private/Windows/NetworkServiceDiscoveryWindows.h
new file mode 100644
--- /dev/null
+++ b/Engine/Plugins/Experimental/NetworkServiceDiscovery/Source/NetworkServiceDiscovery/Private/Windows/NetworkServiceDiscoveryWindows.h
@@ -0,0 +1,96 @@
+// Copyright Epic Games, Inc. All Rights Reserved.
+
+#pragma once
+
+#include "INetworkServiceDiscoveryPlatform.h"
+#include "Containers/Ticker.h"
+
+#if WITH_WINDOWS_DNSSD
+
+struct FBrowseContext;
+
+/**
+ * Windows implementation of network service discovery using native DNS-SD APIs
+ * from windns.h / dnsapi.dll (Windows 10 1803+).
+ *
+ * dnsapi.dll is loaded dynamically at runtime to avoid hard link-time dependencies
+ * and allow graceful fallback on older Windows versions.
+ *
+ * NOTE: Windows DnsServiceBrowse does not provide service removal notifications.
+ * We work around this with periodic re-browsing: each browse cycle collects the
+ * current set of services and diffs against the previous set to detect removals.
+ */
+class FNetworkServiceDiscoveryWindows : public INetworkServiceDiscoveryPlatform
+{
+public:
+	FNetworkServiceDiscoveryWindows();
+	virtual ~FNetworkServiceDiscoveryWindows();
+
+	/** Shared flag used by async callbacks to detect if this object is still alive */
+	TSharedRef<TAtomic<bool>> bIsAlive = MakeShared<TAtomic<bool>>(true);
+
+	/** Try to load dnsapi.dll and resolve all required function pointers. Returns true on success. */
+	static bool LoadDnsApi();
+
+	/** Whether dnsapi.dll was loaded successfully. All operations are no-ops if false. */
+	static bool bDnsApiLoaded;
+
+	// Registration
+	virtual bool RegisterService(const FString& ServiceName, const FString& ServiceType, int32 Port, const TMap<FString, FString>& TxtRecord) override;
+	virtual void UnregisterService(const FString& ServiceName) override;
+	virtual bool IsServiceRegistered(const FString& ServiceName) const override;
+
+	// Discovery
+	virtual bool StartDiscovery(const FString& ServiceType) override;
+	virtual void StopDiscovery() override;
+	virtual bool IsDiscovering() const override;
+	virtual void ResolveService(const FNetworkServiceInfo& Service) override;
+	virtual TArray<FNetworkServiceInfo> GetDiscoveredServices() const override;
+
+public:
+	/** Per-registration state — public for file-static callback access.
+	 *  Tracks both the original instance (for cancelling pending registrations) and the
+	 *  actual instance from the callback (for deregistering active registrations whose
+	 *  name may have been changed by mDNS conflict resolution). */
+	struct FRegistration
+	{
+		void* OriginalInstance = nullptr;  // DNS_SERVICE_INSTANCE* we constructed (for pending cancel)
+		void* ActualInstance = nullptr;    // DNS_SERVICE_INSTANCE* from callback (for active deregister)
+		void* CancelHandle = nullptr;     // DNS_SERVICE_CANCEL*
+		FString ServiceName;
+		uint64 Generation = 0;            // Distinguishes stale callbacks from current registration
+	};
+
+	/** Active registrations keyed by service name */
+	TMap<FString, FRegistration> Registrations;
+
+private:
+	/** Monotonic counter — incremented per RegisterService call to detect stale callbacks */
+	uint64 RegistrationGeneration = 0;
+
+	/** Browse state */
+	void* BrowseCancelHandle = nullptr;  // DNS_SERVICE_CANCEL*
+	TSharedPtr<FBrowseContext> BrowseCallbackContext;  // Ref-counted; shared with async browse callback
+	bool bIsDiscovering = false;
+	FString BrowseServiceType;
+	FString BrowseQueryName;  // Kept alive for async DnsServiceBrowse
+
+	/** Periodic re-browse ticker for detecting service removals */
+	FTSTicker::FDelegateHandle ReBrowseTickerHandle;
+	bool OnReBrowseTick(float DeltaTime);
+
+public:
+	/** Public for file-static callback access */
+	mutable FCriticalSection ServicesLock;
+	TArray<FNetworkServiceInfo> DiscoveredServices;
+	/** Double-buffered service sets for removal detection.
+	 *  BrowseCallback populates CurrentBrowseCycleServices.
+	 *  OnReBrowseTick swaps current→previous, then diffs previous against the new current. */
+	TSet<FString> CurrentBrowseCycleServices;
+	TSet<FString> PreviousBrowseCycleServices;
+	bool bFirstBrowseCycleComplete = false;
+	mutable FCriticalSection ResolveLock;
+	TArray<void*> PendingResolveCancels;
+};
+
+#endif // WITH_WINDOWS_DNSSD
diff --git a/Engine/Plugins/Experimental/NetworkServiceDiscovery/Source/NetworkServiceDiscovery/Public/INetworkServiceDiscovery.h b/Engine/Plugins/Experimental/NetworkServiceDiscovery/Source/NetworkServiceDiscovery/Public/INetworkServiceDiscovery.h
new file mode 100644
--- /dev/null
+++ b/Engine/Plugins/Experimental/NetworkServiceDiscovery/Source/NetworkServiceDiscovery/Public/INetworkServiceDiscovery.h
@@ -0,0 +1,101 @@
+// Copyright Epic Games, Inc. All Rights Reserved.
+
+#pragma once
+
+#include "Modules/ModuleInterface.h"
+#include "Modules/ModuleManager.h"
+#include "NetworkServiceDiscoveryTypes.h"
+
+#define UE_API NETWORKSERVICEDISCOVERY_API
+
+/**
+ * Module interface for cross-platform mDNS/DNS-SD service discovery.
+ *
+ * Supports two primary operations:
+ *   1. Service Registration (host side) - Advertise a service on the local network
+ *   2. Service Discovery (client side) - Browse for and resolve services on the local network
+ *
+ * Platform implementations:
+ *   - Apple (iOS/Mac): NSNetServiceBrowser / NSNetService
+ *   - Android: NsdManager
+ *   - Windows: Native DNS-SD via windns.h / dnsapi.dll (Windows 10 1803+)
+ */
+class INetworkServiceDiscoveryModule : public IModuleInterface
+{
+public:
+
+	static INetworkServiceDiscoveryModule* Get()
+	{
+		return FModuleManager::LoadModulePtr<INetworkServiceDiscoveryModule>("NetworkServiceDiscovery");
+	}
+
+	// ------------------------------------------------------------------
+	//  Service Registration (Host side - e.g. editor on Win/Mac)
+	// ------------------------------------------------------------------
+
+	/**
+	 * Register (advertise) a service on the local network via mDNS.
+	 *
+	 * @param ServiceName   Human-readable name for this service instance
+	 * @param ServiceType   DNS-SD service type, e.g. "_unrealremote._tcp."
+	 * @param Port          Port the service is listening on
+	 * @param TxtRecord     Optional TXT record key-value metadata
+	 * @return true if registration was initiated (actual result comes via OnServiceRegistered/OnDiscoveryError)
+	 */
+	virtual bool RegisterService(
+		const FString& ServiceName,
+		const FString& ServiceType,
+		int32 Port,
+		const TMap<FString, FString>& TxtRecord = TMap<FString, FString>()) = 0;
+
+	/**
+	 * Unregister (stop advertising) a previously registered service.
+	 *
+	 * @param ServiceName   Name of the service to unregister. If empty, unregisters all services.
+	 */
+	virtual void UnregisterService(const FString& ServiceName = FString()) = 0;
+
+	/** Returns true if the named service (or any service, if empty) is currently registered. */
+	virtual bool IsServiceRegistered(const FString& ServiceName = FString()) const = 0;
+
+	// ------------------------------------------------------------------
+	//  Service Discovery (Client side - e.g. iOS/Android app)
+	// ------------------------------------------------------------------
+
+	/**
+	 * Start browsing for services of the given type on the local network.
+	 *
+	 * @param ServiceType   DNS-SD service type to browse for, e.g. "_unrealremote._tcp."
+	 * @return true if browsing was started successfully
+	 */
+	virtual bool StartDiscovery(const FString& ServiceType) = 0;
+
+	/** Stop browsing for services. */
+	virtual void StopDiscovery() = 0;
+
+	/** Returns true if currently browsing for services. */
+	virtual bool IsDiscovering() const = 0;
+
+	/**
+	 * Resolve a discovered service to obtain its IP address and port.
+	 * Result delivered via OnServiceResolved delegate.
+	 *
+	 * @param Service   The service to resolve (as received from OnServiceFound)
+	 */
+	virtual void ResolveService(const FNetworkServiceInfo& Service) = 0;
+
+	/** Get a snapshot of all currently discovered services. */
+	virtual TArray<FNetworkServiceInfo> GetDiscoveredServices() const = 0;
+
+	// ------------------------------------------------------------------
+	//  Delegates
+	// ------------------------------------------------------------------
+
+	virtual FOnServiceFound& OnServiceFound() = 0;
+	virtual FOnServiceLost& OnServiceLost() = 0;
+	virtual FOnServiceResolved& OnServiceResolved() = 0;
+	virtual FOnServiceRegistered& OnServiceRegistered() = 0;
+	virtual FOnDiscoveryError& OnDiscoveryError() = 0;
+};
+
+#undef UE_API
diff --git a/Engine/Plugins/Experimental/NetworkServiceDiscovery/Source/NetworkServiceDiscovery/Public/NetworkServiceDiscoveryTypes.h b/Engine/Plugins/Experimental/NetworkServiceDiscovery/Source/NetworkServiceDiscovery/Public/NetworkServiceDiscoveryTypes.h
new file mode 100644
--- /dev/null
+++ b/Engine/Plugins/Experimental/NetworkServiceDiscovery/Source/NetworkServiceDiscovery/Public/NetworkServiceDiscoveryTypes.h
@@ -0,0 +1,57 @@
+// Copyright Epic Games, Inc. All Rights Reserved.
+
+#pragma once
+
+#include "CoreMinimal.h"
+
+/**
+ * Represents a network service discovered or registered via mDNS/DNS-SD.
+ */
+struct FNetworkServiceInfo
+{
+	/** Human-readable service name, e.g. "MyProject on DESKTOP-ABC" */
+	FString ServiceName;
+
+	/** DNS-SD service type, e.g. "_unrealremote._tcp." */
+	FString ServiceType;
+
+	/** Resolved hostname of the service provider */
+	FString HostName;
+
+	/** Resolved IP address (IPv4 string) */
+	FString Address;
+
+	/** Resolved port number */
+	int32 Port = 0;
+
+	/** TXT record key-value pairs carrying service metadata */
+	TMap<FString, FString> TxtRecord;
+
+	/** Whether this service has been resolved (Address and Port are valid) */
+	bool bIsResolved = false;
+
+	bool operator==(const FNetworkServiceInfo& Other) const
+	{
+		return ServiceName == Other.ServiceName && ServiceType == Other.ServiceType;
+	}
+
+	friend uint32 GetTypeHash(const FNetworkServiceInfo& Info)
+	{
+		return HashCombine(GetTypeHash(Info.ServiceName), GetTypeHash(Info.ServiceType));
+	}
+};
+
+/** Fired when a new service is found during discovery (not yet resolved). */
+DECLARE_MULTICAST_DELEGATE_OneParam(FOnServiceFound, const FNetworkServiceInfo& /*Service*/);
+
+/** Fired when a previously discovered service is no longer available. */
+DECLARE_MULTICAST_DELEGATE_OneParam(FOnServiceLost, const FNetworkServiceInfo& /*Service*/);
+
+/** Fired when a discovered service has been resolved to an address and port. */
+DECLARE_MULTICAST_DELEGATE_OneParam(FOnServiceResolved, const FNetworkServiceInfo& /*Service*/);
+
+/** Fired when a service has been successfully registered/advertised. */
+DECLARE_MULTICAST_DELEGATE_OneParam(FOnServiceRegistered, const FNetworkServiceInfo& /*Service*/);
+
+/** Fired when a discovery or registration error occurs. */
+DECLARE_MULTICAST_DELEGATE_OneParam(FOnDiscoveryError, const FString& /*ErrorMessage*/);
diff --git a/Engine/Plugins/Experimental/RemoteSession/Source/RemoteSession/Private/RemoteSessionModule.cpp b/Engine/Plugins/Experimental/RemoteSession/Source/RemoteSession/Private/RemoteSessionModule.cpp
--- a/Engine/Plugins/Experimental/RemoteSession/Source/RemoteSession/Private/RemoteSessionModule.cpp
+++ b/Engine/Plugins/Experimental/RemoteSession/Source/RemoteSession/Private/RemoteSessionModule.cpp
@@ -8,6 +8,10 @@
 #include "Modules/ModuleManager.h"
 #include "RemoteSessionClient.h"
 
+#if WITH_NETWORK_SERVICE_DISCOVERY
+#include "INetworkServiceDiscovery.h"
+#endif
+
 #if WITH_EDITOR
 	#include "Editor.h"
 	#include "ISettingsModule.h"
@@ -257,6 +261,18 @@ void FRemoteSessionModule::InitHost(const int16 Port /*= 0*/)
 	{
 		Host = NewHost;
 		UE_LOG(LogRemoteSession, Log, TEXT("Started listening on port %d"), SelectedPort);
+
+#if WITH_NETWORK_SERVICE_DISCOVERY
+		// Advertise via mDNS so mobile clients can discover this host
+		if (INetworkServiceDiscoveryModule* Discovery = FModuleManager::GetModulePtr<INetworkServiceDiscoveryModule>("NetworkServiceDiscovery"))
+		{
+			TMap<FString, FString> TxtRecord;
+			TxtRecord.Add(TEXT("version"), IRemoteSessionModule::GetLocalVersion());
+			TxtRecord.Add(TEXT("project"), FApp::GetProjectName());
+			FString ServiceName = FString::Printf(TEXT("%s on %s"), FApp::GetProjectName(), FPlatformProcess::ComputerName());
+			Discovery->RegisterService(ServiceName, TEXT("_unrealremote._tcp."), SelectedPort, TxtRecord);
+		}
+#endif
 	}
 	else
 	{
@@ -266,6 +282,14 @@ void FRemoteSessionModule::InitHost(const int16 Port /*= 0*/)
 
 void FRemoteSessionModule::StopHost(const FString& InReason)
 { 
+#if WITH_NETWORK_SERVICE_DISCOVERY
+	// Stop mDNS advertisement
+	if (INetworkServiceDiscoveryModule* Discovery = FModuleManager::GetModulePtr<INetworkServiceDiscoveryModule>("NetworkServiceDiscovery"))
+	{
+		Discovery->UnregisterService();
+	}
+#endif
+
 	if (IsHostConnected())
 	{
 		Host->Close(InReason);
diff --git a/Engine/Plugins/Experimental/RemoteSession/Source/RemoteSession/RemoteSession.Build.cs b/Engine/Plugins/Experimental/RemoteSession/Source/RemoteSession/RemoteSession.Build.cs
--- a/Engine/Plugins/Experimental/RemoteSession/Source/RemoteSession/RemoteSession.Build.cs
+++ b/Engine/Plugins/Experimental/RemoteSession/Source/RemoteSession/RemoteSession.Build.cs
@@ -1,5 +1,6 @@
 // Copyright Epic Games, Inc. All Rights Reserved.
 
+using System.IO;
 using UnrealBuildTool;
 
 public class RemoteSession : ModuleRules
@@ -48,6 +49,22 @@ public class RemoteSession : ModuleRules
 			}
 		);
 
+		// NetworkServiceDiscovery is optional - used for mDNS host advertisement.
+		// We add the include path directly (not via PrivateIncludePathModuleNames)
+		// to avoid UBT enforcing a plugin dependency, since this is a soft/runtime dependency.
+		{
+			string NSDPublicPath = Path.Combine(PluginDirectory, "..", "NetworkServiceDiscovery", "Source", "NetworkServiceDiscovery", "Public");
+			if (Directory.Exists(NSDPublicPath))
+			{
+				PrivateIncludePaths.Add(NSDPublicPath);
+				PrivateDefinitions.Add("WITH_NETWORK_SERVICE_DISCOVERY=1");
+			}
+			else
+			{
+				PrivateDefinitions.Add("WITH_NETWORK_SERVICE_DISCOVERY=0");
+			}
+		}
+
 		if (Target.bBuildEditor == true)
 		{
 			//reference the module "MyModule"
