Fix NetworkServiceDiscovery Portable Toolchain build on Apple platforms

#Synced-CL 52204226

diff --git a/Engine/Plugins/Experimental/NetworkServiceDiscovery/Source/NetworkServiceDiscovery/NetworkServiceDiscovery.Build.cs b/Engine/Plugins/Experimental/NetworkServiceDiscovery/Source/NetworkServiceDiscovery/NetworkServiceDiscovery.Build.cs
--- a/Engine/Plugins/Experimental/NetworkServiceDiscovery/Source/NetworkServiceDiscovery/NetworkServiceDiscovery.Build.cs
+++ b/Engine/Plugins/Experimental/NetworkServiceDiscovery/Source/NetworkServiceDiscovery/NetworkServiceDiscovery.Build.cs
@@ -8,6 +8,7 @@ public class NetworkServiceDiscovery : ModuleRules
 	public NetworkServiceDiscovery(ReadOnlyTargetRules Target) : base(Target)
 	{
 		PCHUsage = ModuleRules.PCHUsageMode.UseExplicitOrSharedPCHs;
+		bRequiresPlatformSDK = true;
 
 		PublicDependencyModuleNames.AddRange(
 			new string[]
