Fix NoPCH Compile UnrealGame Android

#Synced-CL 52204226

diff --git a/Engine/Plugins/Experimental/NetworkServiceDiscovery/Source/NetworkServiceDiscovery/Private/Android/NetworkServiceDiscoveryAndroid.h b/Engine/Plugins/Experimental/NetworkServiceDiscovery/Source/NetworkServiceDiscovery/Private/Android/NetworkServiceDiscoveryAndroid.h
--- a/Engine/Plugins/Experimental/NetworkServiceDiscovery/Source/NetworkServiceDiscovery/Private/Android/NetworkServiceDiscoveryAndroid.h
+++ b/Engine/Plugins/Experimental/NetworkServiceDiscovery/Source/NetworkServiceDiscovery/Private/Android/NetworkServiceDiscoveryAndroid.h
@@ -4,6 +4,10 @@
 
 #include "INetworkServiceDiscoveryPlatform.h"
 
+#if PLATFORM_ANDROID
+
+#include <jni.h>
+
 /**
  * Android implementation of network service discovery using NsdManager via JNI.
  */
@@ -61,3 +65,5 @@ private:
 
 	static FNetworkServiceDiscoveryAndroid* Instance;
 };
+
+#endif // PLATFORM_ANDROID
