Change NetworkServiceDiscovery warnings on Windows when dnsapi.dll is not available to log
replace emdashes with hyphens

diff --git a/Engine/Plugins/Experimental/NetworkServiceDiscovery/Source/NetworkServiceDiscovery/Private/NetworkServiceDiscoveryModule.cpp b/Engine/Plugins/Experimental/NetworkServiceDiscovery/Source/NetworkServiceDiscovery/Private/NetworkServiceDiscoveryModule.cpp
--- a/Engine/Plugins/Experimental/NetworkServiceDiscovery/Source/NetworkServiceDiscovery/Private/NetworkServiceDiscoveryModule.cpp
+++ b/Engine/Plugins/Experimental/NetworkServiceDiscovery/Source/NetworkServiceDiscovery/Private/NetworkServiceDiscoveryModule.cpp
@@ -51,11 +51,11 @@ TUniquePtr<INetworkServiceDiscoveryPlatform> FNetworkServiceDiscoveryModule::Cre
 	}
 	else
 	{
-		UE_LOGF(LogNetworkServiceDiscovery, Warning, "Windows DNS-SD APIs not available, falling back to Null implementation.");
+		UE_LOGF(LogNetworkServiceDiscovery, Log, "Windows DNS-SD APIs not available, falling back to Null implementation.");
 		return MakeUnique<FNetworkServiceDiscoveryNull>();
 	}
 #else
-	UE_LOGF(LogNetworkServiceDiscovery, Warning, "Using Null implementation - mDNS not supported on this platform.");
+	UE_LOGF(LogNetworkServiceDiscovery, Log, "Using Null implementation - mDNS not supported on this platform.");
 	return MakeUnique<FNetworkServiceDiscoveryNull>();
 #endif
 }
diff --git a/Engine/Plugins/Experimental/NetworkServiceDiscovery/Source/NetworkServiceDiscovery/Private/Windows/NetworkServiceDiscoveryWindows.cpp b/Engine/Plugins/Experimental/NetworkServiceDiscovery/Source/NetworkServiceDiscovery/Private/Windows/NetworkServiceDiscoveryWindows.cpp
--- a/Engine/Plugins/Experimental/NetworkServiceDiscovery/Source/NetworkServiceDiscovery/Private/Windows/NetworkServiceDiscoveryWindows.cpp
+++ b/Engine/Plugins/Experimental/NetworkServiceDiscovery/Source/NetworkServiceDiscovery/Private/Windows/NetworkServiceDiscoveryWindows.cpp
@@ -31,7 +31,7 @@ typedef DWORD (WINAPI* FnDnsServiceResolve)(PDNS_SERVICE_RESOLVE_REQUEST pReques
 typedef DWORD (WINAPI* FnDnsServiceResolveCancel)(PDNS_SERVICE_CANCEL pCancelHandle);
 typedef VOID (WINAPI* FnDnsRecordListFree)(PDNS_RECORD pRecordList, DNS_FREE_TYPE FreeType);
 
-// Static function pointers — populated by LoadDnsApi()
+// Static function pointers - populated by LoadDnsApi()
 static FnDnsServiceConstructInstance  pDnsServiceConstructInstance  = nullptr;
 static FnDnsServiceFreeInstance       pDnsServiceFreeInstance       = nullptr;
 static FnDnsServiceRegister           pDnsServiceRegister           = nullptr;
@@ -57,7 +57,7 @@ bool FNetworkServiceDiscoveryWindows::LoadDnsApi()
 	HMODULE hDnsApi = LoadLibraryW(L"dnsapi.dll");
 	if (!hDnsApi)
 	{
-		UE_LOGF(LogNetworkServiceDiscovery, Warning, "Windows: Failed to load dnsapi.dll — DNS-SD will not be available");
+		UE_LOGF(LogNetworkServiceDiscovery, Log, "Windows: Failed to load dnsapi.dll - DNS-SD will not be available");
 		return false;
 	}
 
@@ -66,7 +66,7 @@ bool FNetworkServiceDiscoveryWindows::LoadDnsApi()
 		p##Name = reinterpret_cast<Fn##Name>(GetProcAddress(hDnsApi, #Name)); \
 		if (!p##Name) \
 		{ \
-			UE_LOGF(LogNetworkServiceDiscovery, Warning, "Windows: dnsapi.dll missing export '%ls' — DNS-SD will not be available", TEXT(#Name)); \
+			UE_LOGF(LogNetworkServiceDiscovery, Warning, "Windows: dnsapi.dll missing export '%ls' - DNS-SD will not be available", TEXT(#Name)); \
 			FreeLibrary(hDnsApi); \
 			return false; \
 		}
@@ -84,7 +84,7 @@ bool FNetworkServiceDiscoveryWindows::LoadDnsApi()
 
 	#undef NSD_LOAD_FUNC
 
-	// Intentionally never FreeLibrary — the DLL must stay loaded for the
+	// Intentionally never FreeLibrary - the DLL must stay loaded for the
 	// lifetime of the process since async callbacks reference its code.
 	bDnsApiLoaded = true;
 	UE_LOGF(LogNetworkServiceDiscovery, Log, "Windows: dnsapi.dll loaded successfully");
@@ -109,7 +109,7 @@ static FString BuildInstanceName(const FString& ServiceName, const FString& Serv
 //  Callback context structs
 // ============================================================================
 
-/** Context for register/deregister callbacks — captures the alive flag at call-site time.
+/** Context for register/deregister callbacks - captures the alive flag at call-site time.
  *  For register calls, ServiceName + Generation identify the registration so the callback
  *  can store the actual DNS_SERVICE_INSTANCE (which may have a different name due to mDNS
  *  conflict resolution). Stale callbacks (from a cancelled/replaced registration) are
@@ -126,12 +126,12 @@ struct FRegisterContext
 	bool bIsDeregister = false;
 };
 
-/** Context for browse callbacks — captures the alive flag at call-site time.
+/** Context for browse callbacks - captures the alive flag at call-site time.
  *  Ref-counted via TSharedPtr: the caller holds one ref (BrowseCallbackContext member)
  *  and a heap-allocated TSharedPtr copy is passed as the PVOID context to the Windows API.
  *  The callback extracts and copies the shared pointer, then deletes the heap wrapper.
  *  This ensures the context stays alive regardless of whether DnsServiceBrowseCancel
- *  triggers a final callback or not — whichever side releases last frees the memory. */
+ *  triggers a final callback or not - whichever side releases last frees the memory. */
 struct FBrowseContext
 {
 	FNetworkServiceDiscoveryWindows* Self;
@@ -151,7 +151,7 @@ static void WINAPI NSD_RegisterCallback(DWORD Status, PVOID Context, PDNS_SERVIC
 	uint64 Generation = Ctx->Generation;
 
 	// For deregister calls, free the service instance now that deregistration is complete.
-	// This must be deferred until the callback because DnsServiceDeRegister is async —
+	// This must be deferred until the callback because DnsServiceDeRegister is async -
 	// freeing the instance before deregistration completes leaves a ghost service on the
 	// network, causing name conflicts (e.g. "MyService (2)") on re-registration.
 	bool bIsDeregister = Ctx->bIsDeregister;
@@ -163,7 +163,7 @@ static void WINAPI NSD_RegisterCallback(DWORD Status, PVOID Context, PDNS_SERVIC
 
 	if (bIsDeregister)
 	{
-		// Deregister completion — the system may provide a pInstance copy, free it if present.
+		// Deregister completion - the system may provide a pInstance copy, free it if present.
 		// Don't store it or broadcast OnServiceRegistered.
 		if (pInstance)
 		{
@@ -184,7 +184,7 @@ static void WINAPI NSD_RegisterCallback(DWORD Status, PVOID Context, PDNS_SERVIC
 
 		UE_LOGF(LogNetworkServiceDiscovery, Log, "Windows: Service registered '%ls'", *Info.ServiceName);
 
-		// Store the callback's DNS_SERVICE_INSTANCE as ActualInstance — it has the real
+		// Store the callback's DNS_SERVICE_INSTANCE as ActualInstance - it has the real
 		// registered name (which may differ from what we requested due to mDNS conflict
 		// resolution). DnsServiceDeRegister needs this actual-name instance to work.
 		// Generation check prevents stale callbacks from overwriting a newer registration.
@@ -199,7 +199,7 @@ static void WINAPI NSD_RegisterCallback(DWORD Status, PVOID Context, PDNS_SERVIC
 					auto* Reg = Self->Registrations.Find(ServiceName);
 					if (Reg && Reg->Generation == Generation)
 					{
-						// This callback matches the current registration — store the actual instance
+						// This callback matches the current registration - store the actual instance
 						if (Reg->OriginalInstance)
 						{
 							pDnsServiceFreeInstance(static_cast<PDNS_SERVICE_INSTANCE>(Reg->OriginalInstance));
@@ -211,7 +211,7 @@ static void WINAPI NSD_RegisterCallback(DWORD Status, PVOID Context, PDNS_SERVIC
 				}
 				if (!bStored && pInstance)
 				{
-					// Stale callback or registration already removed — free the instance
+					// Stale callback or registration already removed - free the instance
 					pDnsServiceFreeInstance(pInstance);
 				}
 				Self->OnServiceRegisteredDelegate.Broadcast(Info);
@@ -236,7 +236,7 @@ static void WINAPI NSD_RegisterCallback(DWORD Status, PVOID Context, PDNS_SERVIC
 
 static void WINAPI NSD_BrowseCallback(DWORD Status, PVOID Context, PDNS_RECORD pDnsRecord)
 {
-	// Extract the shared pointer from the heap wrapper — this bumps the ref count so the
+	// Extract the shared pointer from the heap wrapper - this bumps the ref count so the
 	// FBrowseContext stays alive for the duration of this callback even if the caller
 	// (StopDiscovery/OnReBrowseTick) has already released its ref.
 	TSharedPtr<FBrowseContext>* pSharedCtx = static_cast<TSharedPtr<FBrowseContext>*>(Context);
@@ -261,7 +261,7 @@ static void WINAPI NSD_BrowseCallback(DWORD Status, PVOID Context, PDNS_RECORD p
 		return;
 	}
 
-	// Walk the DNS record list — PTR records point to service instance names
+	// Walk the DNS record list - PTR records point to service instance names
 	TArray<FNetworkServiceInfo> FoundServices;
 	for (PDNS_RECORD pRecord = pDnsRecord; pRecord != nullptr; pRecord = pRecord->pNext)
 	{
@@ -310,7 +310,7 @@ static void WINAPI NSD_BrowseCallback(DWORD Status, PVOID Context, PDNS_RECORD p
 	});
 }
 
-/** Context for resolve callbacks — carries the owner and the heap-allocated query name */
+/** Context for resolve callbacks - carries the owner and the heap-allocated query name */
 struct FResolveContext
 {
 	FNetworkServiceDiscoveryWindows* Self;
@@ -417,13 +417,13 @@ FNetworkServiceDiscoveryWindows::~FNetworkServiceDiscoveryWindows()
 	UnregisterService(FString());
 	StopDiscovery();
 
-	// Cancel any pending resolves — callbacks may still fire but will
+	// Cancel any pending resolves - callbacks may still fire but will
 	// see bIsAlive==false and skip accessing this object
 	FScopeLock Lock(&ResolveLock);
 	for (void* Handle : PendingResolveCancels)
 	{
 		pDnsServiceResolveCancel(static_cast<PDNS_SERVICE_CANCEL>(Handle));
-		// Don't delete here — the cancel triggers the callback which handles cleanup
+		// Don't delete here - the cancel triggers the callback which handles cleanup
 	}
 	PendingResolveCancels.Empty();
 }
@@ -439,7 +439,7 @@ bool FNetworkServiceDiscoveryWindows::RegisterService(const FString& ServiceName
 	FString InstanceName = BuildInstanceName(ServiceName, ServiceType);
 	FString HostName = FString::Printf(TEXT("%s.local"), FPlatformProcess::ComputerName());
 
-	// Build TXT record arrays — FStrings own the data, pointer arrays are views for the API
+	// Build TXT record arrays - FStrings own the data, pointer arrays are views for the API
 	TArray<FString> TxtKeyStrings;
 	TArray<FString> TxtValueStrings;
 	for (const auto& Pair : TxtRecord)
@@ -459,7 +459,7 @@ bool FNetworkServiceDiscoveryWindows::RegisterService(const FString& ServiceName
 	PDNS_SERVICE_INSTANCE pInstance = pDnsServiceConstructInstance(
 		*InstanceName,
 		*HostName,
-		nullptr,    // IPv4 — let the system determine
+		nullptr,    // IPv4 - let the system determine
 		nullptr,    // IPv6
 		(WORD)Port,
 		0,          // priority
@@ -524,7 +524,7 @@ void FNetworkServiceDiscoveryWindows::UnregisterService(const FString& ServiceNa
 	{
 		if (Reg.ActualInstance)
 		{
-			// Case A: Registration completed — deregister with the actual instance (correct name).
+			// Case A: Registration completed - deregister with the actual instance (correct name).
 			// DnsServiceDeRegister matches by instance name, so we must use the instance from
 			// the callback which has the real registered name (possibly renamed by conflict resolution).
 			DNS_SERVICE_REGISTER_REQUEST Request;
@@ -542,13 +542,13 @@ void FNetworkServiceDiscoveryWindows::UnregisterService(const FString& ServiceNa
 			Request.pQueryContext = DeregCtx;
 
 			pDnsServiceDeRegister(&Request, nullptr);
-			// Don't free the instance here — DnsServiceDeRegister is async.
+			// Don't free the instance here - DnsServiceDeRegister is async.
 			// The callback will free it once deregistration actually completes.
 			Reg.ActualInstance = nullptr;
 		}
 		else if (Reg.OriginalInstance)
 		{
-			// Case B: Registration still pending — cancel by handle (no name needed).
+			// Case B: Registration still pending - cancel by handle (no name needed).
 			// DnsServiceRegisterCancel uses the cancel handle to identify the registration,
 			// so it works regardless of what name the system would have assigned.
 			if (Reg.CancelHandle)
@@ -606,7 +606,7 @@ bool FNetworkServiceDiscoveryWindows::StartDiscovery(const FString& ServiceType)
 
 	BrowseServiceType = ServiceType;
 
-	// Build query name: "_unrealremote._tcp.local" — stored as member to keep alive for async browse
+	// Build query name: "_unrealremote._tcp.local" - stored as member to keep alive for async browse
 	FString CleanType = ServiceType;
 	CleanType.RemoveFromEnd(TEXT("."));
 	BrowseQueryName = FString::Printf(TEXT("%s.local"), *CleanType);
@@ -694,7 +694,7 @@ bool FNetworkServiceDiscoveryWindows::OnReBrowseTick(float DeltaTime)
 	PreviousBrowseCycleServices = MoveTemp(CurrentBrowseCycleServices);
 	CurrentBrowseCycleServices.Empty();
 
-	// Only diff after the first full cycle has completed — skip the very first tick
+	// Only diff after the first full cycle has completed - skip the very first tick
 	// where PreviousBrowseCycleServices would be empty and cause false removals.
 	if (bFirstBrowseCycleComplete)
 	{
@@ -722,7 +722,7 @@ bool FNetworkServiceDiscoveryWindows::OnReBrowseTick(float DeltaTime)
 		bFirstBrowseCycleComplete = true;
 	}
 
-	// Issue a new browse — results will populate CurrentBrowseCycleServices via the callback
+	// Issue a new browse - results will populate CurrentBrowseCycleServices via the callback
 	if (BrowseCancelHandle)
 	{
 		pDnsServiceBrowseCancel(static_cast<PDNS_SERVICE_CANCEL>(BrowseCancelHandle));
@@ -767,7 +767,7 @@ bool FNetworkServiceDiscoveryWindows::OnReBrowseTick(float DeltaTime)
 
 void FNetworkServiceDiscoveryWindows::ResolveService(const FNetworkServiceInfo& Service)
 {
-	// Allocate the query name on the heap — DnsServiceResolve stores the pointer
+	// Allocate the query name on the heap - DnsServiceResolve stores the pointer
 	// for the async operation's lifetime. Freed in the resolve callback via FResolveContext.
 	WCHAR* QueryNameBuf = new WCHAR[Service.ServiceName.Len() + 1];
 	FCString::Strncpy(QueryNameBuf, *Service.ServiceName, Service.ServiceName.Len() + 1);
